Fix incorrect casing on file includes

diff --git a/Engine/Plugins/Runtime/Synthesis/Source/Synthesis/EpicSynth1.h b/Engine/Plugins/Runtime/Synthesis/Source/Synthesis/EpicSynth1.h
--- a/Engine/Plugins/Runtime/Synthesis/Source/Synthesis/EpicSynth1.h
+++ b/Engine/Plugins/Runtime/Synthesis/Source/Synthesis/EpicSynth1.h
@@ -6,7 +6,7 @@
 
 #include "DSP/Osc.h"
 #include "DSP/LFO.h"
-#include "DSP/DSP.h"
+#include "DSP/Dsp.h"
 #include "DSP/Envelope.h"
 #include "DSP/Amp.h"
 #include "DSP/DelayStereo.h"
diff --git a/Engine/Source/Runtime/AudioMixer/Private/DSP/Amp.cpp b/Engine/Source/Runtime/AudioMixer/Private/DSP/Amp.cpp
--- a/Engine/Source/Runtime/AudioMixer/Private/DSP/Amp.cpp
+++ b/Engine/Source/Runtime/AudioMixer/Private/DSP/Amp.cpp
@@ -1,7 +1,7 @@
 // Copyright 1998-2017 Epic Games, Inc. All Rights Reserved.
 
 #include "DSP/Amp.h"
-#include "DSP/DSP.h"
+#include "DSP/Dsp.h"
 
 namespace Audio
 {
diff --git a/Engine/Source/Runtime/AudioMixer/Private/DSP/Envelope.cpp b/Engine/Source/Runtime/AudioMixer/Private/DSP/Envelope.cpp
--- a/Engine/Source/Runtime/AudioMixer/Private/DSP/Envelope.cpp
+++ b/Engine/Source/Runtime/AudioMixer/Private/DSP/Envelope.cpp
@@ -1,7 +1,7 @@
 // Copyright 1998-2017 Epic Games, Inc. All Rights Reserved.
 
 #include "DSP/Envelope.h"
-#include "DSP/DSP.h"
+#include "DSP/Dsp.h"
 
 namespace Audio
 {
diff --git a/Engine/Source/Runtime/AudioMixer/Private/DSP/LFO.cpp b/Engine/Source/Runtime/AudioMixer/Private/DSP/LFO.cpp
--- a/Engine/Source/Runtime/AudioMixer/Private/DSP/LFO.cpp
+++ b/Engine/Source/Runtime/AudioMixer/Private/DSP/LFO.cpp
@@ -1,7 +1,7 @@
 // Copyright 1998-2017 Epic Games, Inc. All Rights Reserved.
 
 #include "DSP/LFO.h"
-#include "DSP/DSP.h"
+#include "DSP/Dsp.h"
 
 namespace Audio
 {
