[Release-4.20] fix issue where Synth Sample Player would not play a sound wave set in the property inspector. #rb none #jira none #fyi aaron.mcleran

diff --git a/Engine/Plugins/Runtime/Synthesis/Source/Synthesis/Private/SynthComponentWaveTable.cpp b/Engine/Plugins/Runtime/Synthesis/Source/Synthesis/Private/SynthComponentWaveTable.cpp
--- a/Engine/Plugins/Runtime/Synthesis/Source/Synthesis/Private/SynthComponentWaveTable.cpp
+++ b/Engine/Plugins/Runtime/Synthesis/Source/Synthesis/Private/SynthComponentWaveTable.cpp
@@ -25,6 +25,12 @@ bool USynthSamplePlayer::Init(int32& SampleRate)
 
 	SampleBufferReader.Init(SampleRate);
 	SoundWaveLoader.Init(GetAudioDevice());
+
+	if (SoundWave != nullptr)
+	{
+		SoundWaveLoader.LoadSoundWave(SoundWave);
+	}
+
 	return true;
 }
 
