Copying Dev-Audio to UE-Main from CL 6717121

diff --git a/Engine/Config/BaseInput.ini b/Engine/Config/BaseInput.ini
--- a/Engine/Config/BaseInput.ini
+++ b/Engine/Config/BaseInput.ini
@@ -198,10 +198,13 @@ MaxScrollbackSize=1024
 +ManualAutoCompleteList=(Command="ShowDebugToggleSubCategory FULLGRAPH",Desc="With ShowDebug Animation: Toggles graph display between active nodes only and all nodes")
 +ManualAutoCompleteList=(Command="ShowDebugToggleSubCategory FULLBLENDSPACEDISPLAY",Desc="With ShowDebug Animation: Toggles display of sample blend weights on blendspaces")
 +ManualAutoCompleteList=(Command="ShowDebugForReticleTargetToggle ",Desc="<DesiredClass> Toggles 'ShowDebug' from showing debug info between reticle target actor (of subclass <DesiredClass>) and camera view target")
-+ManualAutoCompleteList=(Command="Stat SOUNDMIXES",Desc="Shows active SoundMixes")
-+ManualAutoCompleteList=(Command="Stat SOUNDWAVES",Desc="Shows active SoundWaves")
-+ManualAutoCompleteList=(Command="Stat SOUNDCUES",Desc="Shows active SoundCues")
-+ManualAutoCompleteList=(Command="Stat SOUNDS",Desc="<?> <sort=distance|class|name|waves|default> <-debug> <off> Shows active SoundCues and SoundWaves")
++ManualAutoCompleteList=(Command="Stat SoundCues",Desc="Shows active SoundCues")
++ManualAutoCompleteList=(Command="Stat SoundMixes",Desc="Shows active SoundMixes")
++ManualAutoCompleteList=(Command="Stat SoundModulators",Desc="Shows modulator debug info as provided by active audio modulation plugin")
++ManualAutoCompleteList=(Command="Stat SoundModulatorsHelp",Desc="Shows modulator debug help provided by active audio modulation plugin")
++ManualAutoCompleteList=(Command="Stat SoundReverb",Desc="Shows active SoundReverb")
++ManualAutoCompleteList=(Command="Stat SoundWaves",Desc="Shows active SoundWaves")
++ManualAutoCompleteList=(Command="Stat Sounds",Desc="<?> <sort=distance|class|name|waves|default> <-debug> <off> Shows active SoundCues and SoundWaves")
 +ManualAutoCompleteList=(Command="Audio3dVisualize",Desc="Shows locations of sound sources playing (white text) and their left and right channel locations respectively (red and green).  Virtualized loops (if enabled) display in blue.")
 +ManualAutoCompleteList=(Command="StartMovieCapture",Desc=)
 +ManualAutoCompleteList=(Command="StopMovieCapture",Desc=)
diff --git a/Engine/Plugins/Runtime/MIDIDevice/Source/MIDIDevice/MIDIDeviceInputController.cpp b/Engine/Plugins/Runtime/MIDIDevice/Source/MIDIDevice/MIDIDeviceInputController.cpp
--- a/Engine/Plugins/Runtime/MIDIDevice/Source/MIDIDevice/MIDIDeviceInputController.cpp
+++ b/Engine/Plugins/Runtime/MIDIDevice/Source/MIDIDevice/MIDIDeviceInputController.cpp
@@ -143,7 +143,7 @@ void UMIDIDeviceInputController::ProcessIncomingMIDIEvents()
 
 					case EMIDIEventType::PitchBend:
 					{
-						int32 pitchBend = (PMMessageData1 & 0x7F) << 7;
+						int32 pitchBend = (PMMessageData2 & 0x7F) << 7;
 
 						pitchBend |= (PMMessageData1 & 0x7F);
 
diff --git a/Engine/Plugins/Runtime/OSC/OSC.uplugin b/Engine/Plugins/Runtime/OSC/OSC.uplugin
new file mode 100644
--- /dev/null
+++ b/Engine/Plugins/Runtime/OSC/OSC.uplugin
@@ -0,0 +1,29 @@
+{
+	"FileVersion" : 3,
+	"Version" : 1,
+	"VersionName" : "1.0",
+	"FriendlyName" : "OSC Support",
+	"Description" : "Allows you to send and receive OSC events through a simple API in either C++ or Blueprints",
+	"Category" : "Input Devices",
+	"CreatedBy" : "Epic Games, Inc.",
+	"CreatedByURL" : "http://epicgames.com",
+	"DocsURL" : "",
+	"MarketplaceURL" : "",
+	"SupportURL" : "",
+	"EnabledByDefault" : false,
+	"CanContainContent" : false,
+	"IsBetaVersion" : true,
+	"Installed" : false,
+	"Modules" :
+	[
+		{
+			"Name": "OSC",
+			"Type": "Runtime",
+			"LoadingPhase": "Default",
+			"WhitelistPlatforms": [
+				"Win64",
+				"Win32"
+			]
+		}
+	]
+}
\ No newline at end of file
diff --git a/Engine/Plugins/Runtime/OSC/Source/OSC/Classes/OSCBundle.h b/Engine/Plugins/Runtime/OSC/Source/OSC/Classes/OSCBundle.h
new file mode 100644
--- /dev/null
+++ b/Engine/Plugins/Runtime/OSC/Source/OSC/Classes/OSCBundle.h
@@ -0,0 +1,79 @@
+// Copyright 1998-2019 Epic Games, Inc. All Rights Reserved.
+#pragma once 
+
+#include "OSCPacket.h"
+#include "OSCBundle.generated.h"
+
+class FOSCBundlePacket : public FOSCPacket
+{
+public:
+	FOSCBundlePacket();
+	~FOSCBundlePacket();
+	
+	/** Set the bundle time tag. */
+	void SetTimeTag(uint64 NewTimeTag);
+
+	/** Get the bundle time tag. */
+	uint64 GetTimeTag() const;
+
+	/** Get the number of OSC packets stored in the bundle. */
+	int32 GetNumPackets() const { return Packets.Num(); }
+
+	/** Add OSC packet in the bundle. */
+	void AddPacket(TSharedPtr<FOSCPacket> Packet);
+	
+	/** Get OSC packet by index. */
+	TSharedPtr<FOSCPacket> GetPacket(int Index) const { return Packets[Index]; }
+	
+	/** Clear bundle packet list. */
+	void Clear() { Packets.Reset(); }
+
+	/** Returns true. */
+	virtual bool IsBundle() { return true; }
+
+	/** Write Bundle data into the OSC stream. */
+	virtual void WriteData(FOSCStream& Stream);
+
+	/** Read Bundle data from an OSC stream*/
+	virtual void ReadData(FOSCStream& Stream);
+
+private:
+
+	/** Bundle OSC packets. */
+	TArray<TSharedPtr<FOSCPacket>> Packets;
+	
+	/** Bundle time tag. */
+	FOSCType TimeTag;
+};
+
+USTRUCT(BlueprintType)
+struct OSC_API FOSCBundle
+{
+	GENERATED_USTRUCT_BODY()
+
+	FOSCBundle() = default;
+	~FOSCBundle()
+	{
+		Packet.Reset();
+	}
+
+	bool IsValid() const 
+	{
+		return Packet.IsValid();
+	}
+
+	void SetPacket(TSharedPtr<FOSCBundlePacket> InPacket)
+	{
+		this->Packet = InPacket;
+	}
+
+	const TSharedPtr<FOSCBundlePacket> GetPacket() const
+	{
+		return Packet;
+	}
+
+	const TSharedPtr<FOSCBundlePacket> GetOrCreatePacket();
+
+private:
+	TSharedPtr<FOSCBundlePacket> Packet;
+};
\ No newline at end of file
diff --git a/Engine/Plugins/Runtime/OSC/Source/OSC/Classes/OSCClient.h b/Engine/Plugins/Runtime/OSC/Source/OSC/Classes/OSCClient.h
new file mode 100644
--- /dev/null
+++ b/Engine/Plugins/Runtime/OSC/Source/OSC/Classes/OSCClient.h
@@ -0,0 +1,49 @@
+// Copyright 1998-2019 Epic Games, Inc. All Rights Reserved.
+#pragma once
+
+#include "CoreMinimal.h"
+#include "UObject/Object.h"
+#include "Interfaces/IPv4/IPv4Address.h"
+#include "IPAddress.h"
+#include "OSCMessage.h"
+#include "OSCBundle.h"
+#include "OSCClient.generated.h"
+
+class FSocket;
+
+UCLASS(BlueprintType)
+class OSC_API UOSCClient : public UObject
+{
+	GENERATED_BODY()
+
+public:
+	UOSCClient();
+	virtual ~UOSCClient() {}
+
+	/** Set's the OSC Client IP address and port. */
+	void SetTarget(FIPv4Address InIPAddress, uint32_t Port);
+
+	/** Send OSC message over the network with a specific address. */
+	UFUNCTION(BlueprintCallable, Category = "Audio|OSC")
+	void SendOSCMessage(FString Address, UPARAM(ref) FOSCMessage& Msg);
+
+	/** Send OSC Bundle over the network. */
+	UFUNCTION(BlueprintCallable, Category = "Audio|OSC")
+	void SendOSCBundle(UPARAM(ref) FOSCBundle& Bundle);
+
+protected:
+
+	void BeginDestroy() override;
+	
+	/** Stop and tidy up network socket. */
+	void Stop();
+	
+	/** Send OSC packet data. */
+	void SendPacket(FOSCPacket* Packet, char* Data, int32 Size);
+
+	/** Socket used to send the OSC packets. */
+	FSocket* Socket;
+
+	/** IP Address used by socket. */
+	TSharedPtr<FInternetAddr> IPAddress;
+};
diff --git a/Engine/Plugins/Runtime/OSC/Source/OSC/Classes/OSCLog.h b/Engine/Plugins/Runtime/OSC/Source/OSC/Classes/OSCLog.h
new file mode 100644
--- /dev/null
+++ b/Engine/Plugins/Runtime/OSC/Source/OSC/Classes/OSCLog.h
@@ -0,0 +1,8 @@
+// Copyright 1998-2019 Epic Games, Inc. All Rights Reserved.
+
+#pragma once
+
+#include "CoreMinimal.h"
+
+DECLARE_LOG_CATEGORY_EXTERN( OSCLog, Verbose, All );
+
diff --git a/Engine/Plugins/Runtime/OSC/Source/OSC/Classes/OSCMessage.h b/Engine/Plugins/Runtime/OSC/Source/OSC/Classes/OSCMessage.h
new file mode 100644
--- /dev/null
+++ b/Engine/Plugins/Runtime/OSC/Source/OSC/Classes/OSCMessage.h
@@ -0,0 +1,86 @@
+// Copyright 1998-2019 Epic Games, Inc. All Rights Reserved.
+#pragma once
+
+#include "CoreMinimal.h"
+#include "UObject/ObjectMacros.h"
+#include "UObject/Object.h"
+#include "UObject/WeakObjectPtr.h"
+#include "OSCPacket.h"
+#include "OSCMessage.generated.h"
+
+class FOSCStream;
+class FOSCMessagePacket : public FOSCPacket
+{
+public:
+	FOSCMessagePacket();
+	~FOSCMessagePacket();
+
+	/** Set OSC message address. */
+	void SetAddress(const FString& InAddress) { this->Address = InAddress; }
+	
+	/** Get OSC message address. */
+	const FString& GetAddress() const { return Address; }
+	
+	/** Add OSC Type as an argument into the OSC message. */
+	void AddArgument(FOSCType type) { Arguments.Add(type); }
+
+	/** Get argument by index. */
+	const FOSCType& GetArgument(int i) const { return Arguments[i]; }
+	
+	/** Get number of arguments. */
+	int GetNumArguments() const { return Arguments.Num(); }
+
+	/** Clear OSC Message. */
+	void Clear() { Arguments.Reset(); }
+
+	/** Returns true to indicate its an OSC message. */
+	virtual bool IsMessage() override { return true; }
+
+	/** Write message data into an OSC stream. */
+	virtual void WriteData(FOSCStream& Stream);
+
+	/** Read message data from an OSC stream. */
+	virtual void ReadData(FOSCStream& Stream);
+
+private:
+
+	/** Message address. */
+	FString Address; 
+
+	/** List of argument types. */
+	TArray<FOSCType> Arguments;
+};
+
+USTRUCT(BlueprintType)
+struct OSC_API FOSCMessage
+{
+	GENERATED_USTRUCT_BODY()
+
+	FOSCMessage() = default;
+	~FOSCMessage()
+	{
+		Packet.Reset();
+	}
+
+	bool IsValid() const
+	{
+		return Packet.IsValid();
+	}
+
+	void SetPacket(TSharedPtr<FOSCMessagePacket> InPacket)
+	{
+		this->Packet = InPacket;
+	}
+
+	const TSharedPtr<FOSCMessagePacket> GetPacket() const
+	{
+		return Packet;
+	}
+
+	const TSharedPtr<FOSCMessagePacket> GetOrCreatePacket();
+
+private:
+
+	/** The OSC Message packet */
+	TSharedPtr<FOSCMessagePacket> Packet;
+};
\ No newline at end of file
diff --git a/Engine/Plugins/Runtime/OSC/Source/OSC/Classes/OSCPacket.h b/Engine/Plugins/Runtime/OSC/Source/OSC/Classes/OSCPacket.h
new file mode 100644
--- /dev/null
+++ b/Engine/Plugins/Runtime/OSC/Source/OSC/Classes/OSCPacket.h
@@ -0,0 +1,27 @@
+// Copyright 1998-2019 Epic Games, Inc. All Rights Reserved.
+#pragma once
+
+#include "OSCTypes.h"
+
+class FOSCStream;
+class FOSCPacket
+{
+public:
+	FOSCPacket() = default;
+	virtual ~FOSCPacket() = default;
+	
+	/** Write packet data into stream */
+	virtual void WriteData(FOSCStream& Stream) {}
+	
+	/** Read packet data from stream */
+	virtual void ReadData(FOSCStream& Stream) {}
+	
+	/** Returns true if packet is message */
+	virtual bool IsMessage() { return false; }
+	
+	/** Returns true if packet is bundle */
+	virtual bool IsBundle() { return false; }
+
+	/** Create an OSC packet according to the input data. */
+	static TSharedPtr <FOSCPacket> CreatePacket(const char* PacketType);
+};
diff --git a/Engine/Plugins/Runtime/OSC/Source/OSC/Classes/OSCStream.h b/Engine/Plugins/Runtime/OSC/Source/OSC/Classes/OSCStream.h
new file mode 100644
--- /dev/null
+++ b/Engine/Plugins/Runtime/OSC/Source/OSC/Classes/OSCStream.h
@@ -0,0 +1,93 @@
+// Copyright 1998-2019 Epic Games, Inc. All Rights Reserved.
+#pragma once
+
+#include "CoreMinimal.h"
+
+class FOSCStream
+{
+public:
+	FOSCStream(char* Data, int Size);
+	FOSCStream();
+	virtual ~FOSCStream();
+
+	/** Get stream bufer. */
+	const char* GetBuffer() { return Data; }
+	
+	/** Get stream length */
+	int GetLength() { return Position; }
+	
+	/** Returns true if stream has reached the end. */
+	bool HasReachedEnd(){ return !(Position < Size); }
+	
+	/** Get current stream position. */
+	int GetPosition() { return Position; }
+
+	/** Set stream position. */
+	void SetPosition(int InPosition) { this->Position = InPosition; }
+
+	/** Read Char from the stream. */
+	char ReadChar();
+
+	/** Write Char into the stream. */
+	void WriteChar(char Char);
+
+	/** Read Int32 from the stream. */
+	int32 ReadInt32();
+
+	/** Write Int32 into the stream. */
+	void WriteInt32(int32 Value);
+	
+	/** Read Double from the stream. */
+	double ReadDouble();
+
+	/** Write Double into the stream. */
+	void WriteDouble(uint64 Value);
+	
+	/** Read Int64 from the stream. */
+	int64 ReadInt64();
+
+	/** Write Int64 into the stream. */
+	void WriteInt64(int64 Value);
+
+	/** Read UInt64 from the stream. */
+	uint64 ReadUInt64();
+
+	/** Write UInt64 into the stream. */
+	void WriteUInt64(uint64 Value);
+
+	/** Read Float from the stream. */
+	float ReadFloat();
+		
+	/** Write Int64 into the stream. */
+	void WriteFloat(float Value);
+	
+	/** Read String from the stream. */
+	FString ReadString();
+
+	/** Write String into the stream. */
+	void WriteString(FString String);
+	
+	/** Read Blob from the stream. */
+	TArray<uint8> ReadBlob();
+
+	/** Write Blob into the stream. */
+	void WriteBlob(TArray<uint8>& Blob);
+
+private:
+	
+	/** Read data rom a buffer. */
+	int Read(void* Buffer, int ToRead);
+	
+	/** Write data into buffer. */
+	int Write(void* Buffer, int ToWrite);
+
+	/** Stream data. */
+	char* Data;
+
+	/** Current buffer position. */
+	int Position;
+
+	/** Stream size. */
+	int32 Size;
+};
+
diff --git a/Engine/Plugins/Runtime/OSC/Source/OSC/Private/OSCBundle.cpp b/Engine/Plugins/Runtime/OSC/Source/OSC/Private/OSCBundle.cpp
new file mode 100644
--- /dev/null
+++ b/Engine/Plugins/Runtime/OSC/Source/OSC/Private/OSCBundle.cpp
@@ -0,0 +1,83 @@
+// Copyright 1998-2019 Epic Games, Inc. All Rights Reserved.
+#include "OSCBundle.h"
+#include "OSCStream.h"
+
+FOSCBundlePacket::FOSCBundlePacket() 
+: FOSCPacket(), TimeTag(0) 
+{}
+
+FOSCBundlePacket::~FOSCBundlePacket() 
+{}
+
+void FOSCBundlePacket::SetTimeTag(uint64 NewTimeTag) 
+{ 
+	TimeTag = FOSCType(NewTimeTag); 
+}
+
+uint64 FOSCBundlePacket::GetTimeTag() const 
+{ 
+	return TimeTag.GetTimeTag(); 
+}
+
+void FOSCBundlePacket::AddPacket(TSharedPtr<FOSCPacket> Packet)
+{
+	Packets.Add(Packet);
+}
+
+void FOSCBundlePacket::WriteData(FOSCStream& Stream)
+{
+	// Write bundle
+	Stream.WriteString(FString("#bundle"));
+
+	// Write bundle time tag
+	Stream.WriteUInt64(GetTimeTag());
+
+	// for each packet
+	for (auto& packet : Packets)
+	{
+		int sizePosition = Stream.GetPosition();
+		Stream.WriteInt32(0);
+
+		int prePos = Stream.GetPosition();
+		packet->WriteData(Stream);
+		int postPos = Stream.GetPosition();
+
+		Stream.SetPosition(sizePosition);
+		Stream.WriteInt32(postPos - prePos);
+		Stream.SetPosition(postPos);
+	}
+}
+
+void FOSCBundlePacket::ReadData(FOSCStream& Stream)
+{
+	// Read address
+	FString address = Stream.ReadString();
+
+	// read time tag
+	TimeTag = FOSCType(Stream.ReadUInt64());
+
+	// While we still have data
+	while (!Stream.HasReachedEnd())
+	{
+		int32 sizeOfPacket = Stream.ReadInt32();
+
+		TSharedPtr<FOSCPacket> packet = FOSCPacket::CreatePacket(Stream.GetBuffer());
+
+		if (packet.IsValid())
+		{
+			packet->ReadData(Stream);
+			Packets.Add(packet);
+		}		
+	}
+}
+
+const TSharedPtr<FOSCBundlePacket> FOSCBundle::GetOrCreatePacket()
+{
+	if (!Packet.IsValid())
+	{
+		Packet = MakeShareable(new FOSCBundlePacket());
+	}
+
+	return Packet;
+}
+
diff --git a/Engine/Plugins/Runtime/OSC/Source/OSC/Private/OSCClient.cpp b/Engine/Plugins/Runtime/OSC/Source/OSC/Private/OSCClient.cpp
new file mode 100644
--- /dev/null
+++ b/Engine/Plugins/Runtime/OSC/Source/OSC/Private/OSCClient.cpp
@@ -0,0 +1,79 @@
+// Copyright 1998-2019 Epic Games, Inc. All Rights Reserved.
+#include "OSCClient.h"
+
+#include "Common/UdpSocketBuilder.h"
+#include "Sockets.h"
+
+#include "OSCStream.h"
+#include "OSCBundle.h"
+#include "OSCLog.h"
+
+UOSCClient::UOSCClient()
+	: Socket(FUdpSocketBuilder(TEXT("OscSender")).Build())
+{
+	//
+}
+
+void UOSCClient::SetTarget(FIPv4Address InIPAddress, uint32_t Port)
+{
+	this->IPAddress = ISocketSubsystem::Get(PLATFORM_SOCKETSUBSYSTEM)->CreateInternetAddr();
+
+	this->IPAddress->SetIp(InIPAddress.Value);
+	this->IPAddress->SetPort(Port);
+
+	UE_LOG(OSCLog, Display, TEXT("Set target port: %d"), Port);
+}
+
+void UOSCClient::SendPacket(FOSCPacket* Packet, char* Data, int32 Size)
+{
+	if (IPAddress)
+	{
+		FOSCStream stream = FOSCStream(Data, Size);
+
+		Packet->WriteData(stream);
+
+		int32 length = stream.GetPosition();
+
+		int32 bytesSent = 0;
+		while (length > 0)
+		{
+			Socket->SendTo((uint8*)Data, length, bytesSent, *IPAddress);
+
+			length -= bytesSent;
+			Data += bytesSent;
+		}
+	}
+}
+
+void UOSCClient::Stop()
+{
+	if (Socket)
+	{
+		ISocketSubsystem::Get(PLATFORM_SOCKETSUBSYSTEM)->DestroySocket(Socket);
+		Socket = nullptr;
+	}
+}
+
+void UOSCClient::BeginDestroy()
+{
+	Stop();
+	Super::BeginDestroy();
+}
+
+void UOSCClient::SendOSCMessage(FString Address, UPARAM(ref) FOSCMessage& Msg)
+{
+	static const int OUTPUT_BUFFER_SIZE = 1024;
+	char buffer[OUTPUT_BUFFER_SIZE];
+
+	Msg.GetPacket()->SetAddress(Address);
+
+	SendPacket(&(*Msg.GetPacket()), &buffer[0], OUTPUT_BUFFER_SIZE);
+}
+
+void UOSCClient::SendOSCBundle(UPARAM(ref) FOSCBundle& Bundle)
+{
+	static const int OUTPUT_BUFFER_SIZE = 1024;
+	char buffer[OUTPUT_BUFFER_SIZE];
+
+	SendPacket(&(*Bundle.GetPacket()), &buffer[0], OUTPUT_BUFFER_SIZE);
+}
diff --git a/Engine/Plugins/Runtime/OSC/Source/OSC/Private/OSCMessage.cpp b/Engine/Plugins/Runtime/OSC/Source/OSC/Private/OSCMessage.cpp
new file mode 100644
--- /dev/null
+++ b/Engine/Plugins/Runtime/OSC/Source/OSC/Private/OSCMessage.cpp
@@ -0,0 +1,147 @@
+// Copyright 1998-2019 Epic Games, Inc. All Rights Reserved.
+#include "OSCMessage.h"
+#include "OSCStream.h"
+
+FOSCMessagePacket::FOSCMessagePacket() 
+	: FOSCPacket() 
+{}
+
+FOSCMessagePacket::~FOSCMessagePacket() 
+{}
+
+
+void FOSCMessagePacket::WriteData(FOSCStream& Stream)
+{
+	// Begin writing data 
+	Stream.WriteString(Address);
+
+	// write type tags
+	FString tagTypes = ",";
+	for (int i = 0; i < GetNumArguments(); i++)
+	{
+		tagTypes += static_cast<char>(GetArgument(i).GetTypeTag());
+	}
+
+	Stream.WriteString(tagTypes);
+
+	for (FOSCType type : Arguments)
+	{
+		switch (type.GetTypeTag())
+		{
+		case EOSCTypeTag::CHAR:
+			Stream.WriteChar(type.GetChar());
+			break;
+		case EOSCTypeTag::INT32:
+			Stream.WriteInt32(type.GetInt32());
+			break;
+		case EOSCTypeTag::FLOAT:
+			Stream.WriteFloat(type.GetFloat());
+			break;
+		case EOSCTypeTag::DOUBLE:
+			Stream.WriteDouble(type.GetDouble());
+			break;
+		case EOSCTypeTag::INT64:
+			Stream.WriteInt64(type.GetInt64());
+			break;
+		case EOSCTypeTag::TIME:
+			Stream.WriteUInt64(type.GetTimeTag());
+			break;
+		case EOSCTypeTag::STRING:
+			Stream.WriteString(type.GetString());
+			break;
+		case EOSCTypeTag::BLOB:
+		{
+			TArray<uint8> blob = type.GetBlob();
+			Stream.WriteBlob(blob);
+		}
+		break;
+		case EOSCTypeTag::COLOR:
+			Stream.WriteInt32(type.GetColor().GetInt32());
+			break;
+		case EOSCTypeTag::TRUE:
+		case EOSCTypeTag::FALSE:
+		case EOSCTypeTag::NIL:
+		case EOSCTypeTag::INFINITUM:
+			// No values are written for these types
+			break;
+		default:
+			// Argument is not supported 
+			unimplemented();
+			break;
+		}
+	}
+}
+
+void FOSCMessagePacket::ReadData(FOSCStream& Stream)
+{
+	// Read Address
+	Address = Stream.ReadString();
+
+	// Read string of tags
+	TArray<TCHAR> tagTypes = Stream.ReadString().GetCharArray();
+
+	// Skip the first argument which is ','
+	for (int i=1; i<tagTypes.Num(); i++)
+	{
+		switch ((EOSCTypeTag)tagTypes[i])
+		{
+		case EOSCTypeTag::CHAR:
+			AddArgument(FOSCType(Stream.ReadChar()));
+			break;
+		case EOSCTypeTag::INT32:
+			AddArgument(FOSCType(Stream.ReadInt32()));
+			break;
+		case EOSCTypeTag::FLOAT:
+			AddArgument(FOSCType(Stream.ReadFloat()));
+			break;
+		case EOSCTypeTag::DOUBLE:
+			AddArgument(FOSCType(Stream.ReadDouble()));
+			break;
+		case EOSCTypeTag::INT64:
+			AddArgument(FOSCType(Stream.ReadInt64()));
+			break;
+		case EOSCTypeTag::TRUE:
+			AddArgument(FOSCType(true));
+			break;
+		case EOSCTypeTag::FALSE:
+			AddArgument(FOSCType(false));
+			break;
+		case EOSCTypeTag::NIL:
+			AddArgument(FOSCType(EOSCTypeTag::NIL));
+			break;
+		case EOSCTypeTag::INFINITUM:
+			AddArgument(FOSCType(EOSCTypeTag::INFINITUM));
+			break;
+		case EOSCTypeTag::TIME:
+			AddArgument(FOSCType(Stream.ReadUInt64()));
+			break;
+		case EOSCTypeTag::STRING:
+			AddArgument(FOSCType(Stream.ReadString()));
+			break;
+		case EOSCTypeTag::BLOB:
+			AddArgument(FOSCType(Stream.ReadBlob()));
+			break;
+		case EOSCTypeTag::COLOR:
+			AddArgument(FOSCType(FOSCColor(Stream.ReadInt32())));
+			break;
+		case EOSCTypeTag::TERMINATE:
+			Stream.ReadChar();
+			break;
+
+		default:
+			// Argument is not supported 
+			unimplemented();
+			break;
+		}
+	}
+}
+
+const TSharedPtr<FOSCMessagePacket> FOSCMessage::GetOrCreatePacket()
+{
+	if (!Packet.IsValid())
+	{
+		Packet = MakeShareable(new FOSCMessagePacket());
+	}
+
+	return Packet;
+}
\ No newline at end of file
diff --git a/Engine/Plugins/Runtime/OSC/Source/OSC/Private/OSCPacket.cpp b/Engine/Plugins/Runtime/OSC/Source/OSC/Private/OSCPacket.cpp
new file mode 100644
--- /dev/null
+++ b/Engine/Plugins/Runtime/OSC/Source/OSC/Private/OSCPacket.cpp
@@ -0,0 +1,18 @@
+// Copyright 1998-2019 Epic Games, Inc. All Rights Reserved.
+#include "OSCPacket.h"
+#include "OSCMessage.h"
+#include "OSCBundle.h"
+
+TSharedPtr<FOSCPacket> FOSCPacket::CreatePacket(const char* PacketType)
+{
+	if (PacketType[0] == '/')
+	{
+		return  MakeShareable(new FOSCMessagePacket());
+	}
+	else if (PacketType[0] == '#')
+	{
+		return  MakeShareable(new FOSCBundlePacket());
+	}
+
+	return nullptr;
+}
\ No newline at end of file
diff --git a/Engine/Plugins/Runtime/OSC/Source/OSC/Private/OSCServer.cpp b/Engine/Plugins/Runtime/OSC/Source/OSC/Private/OSCServer.cpp
new file mode 100644
--- /dev/null
+++ b/Engine/Plugins/Runtime/OSC/Source/OSC/Private/OSCServer.cpp
@@ -0,0 +1,130 @@
+// Copyright 1998-2019 Epic Games, Inc. All Rights Reserved.
+#include "OSCServer.h"
+
+#include "Runtime/Core/Public/Async/TaskGraphInterfaces.h"
+#include "Sockets.h"
+#include "Common/UdpSocketBuilder.h"
+#include "Common/UdpSocketReceiver.h"
+
+#include "OSCStream.h"
+#include "OSCMessage.h"
+#include "OSCBundle.h"
+#include "OSCLog.h"
+
+#define OSC_BUFFER_FREE 0
+#define OSC_BUFFER_LOCKED 1
+
+UOSCServer::UOSCServer()
+: OSCPackets(1024)
+, OSCBufferLock(0)
+{
+}
+
+UOSCServer::~UOSCServer()
+{
+	delete SocketReceiver;
+	SocketReceiver = nullptr;
+
+	if (Socket)
+	{
+		ISocketSubsystem::Get(PLATFORM_SOCKETSUBSYSTEM)->DestroySocket(Socket);
+		Socket = nullptr;
+	}
+}
+
+void UOSCServer::Listen(FIPv4Address IPAddress, uint32_t Port, bool MulticastLoopback)
+{
+	FUdpSocketBuilder builder(TEXT("OscListener"));
+	builder.BoundToPort(Port);
+	if (IPAddress.IsMulticastAddress())
+	{
+		builder.JoinedToGroup(IPAddress);
+		if (MulticastLoopback)
+		{
+			builder.WithMulticastLoopback();
+		}
+	}
+	else
+	{
+		builder.BoundToAddress(IPAddress);
+	}
+
+	Socket = builder.Build();
+	if (Socket)
+	{
+		SocketReceiver = new FUdpSocketReceiver(Socket, FTimespan::FromMilliseconds(100), TEXT("OSCListener"));
+		SocketReceiver->OnDataReceived().BindUObject(this, &UOSCServer::Callback);
+		SocketReceiver->Start();
+
+		UE_LOG(OSCLog, Display, TEXT("Listening to port %d"), Port);
+	}
+	else
+	{
+		UE_LOG(OSCLog, Warning, TEXT("Cannot listen to port %d"), Port);
+	}
+}
+
+void UOSCServer::Callback(const FArrayReaderPtr& Data, const FIPv4Endpoint& Endpoint)
+{
+	TSharedPtr <FOSCPacket> packet = FOSCPacket::CreatePacket((char *)Data->GetData());
+
+	if (packet.IsValid())
+	{
+		// create OSC stream
+		FOSCStream stream = FOSCStream((char *)Data->GetData(), Data->Num());
+
+		packet->ReadData(stream);
+
+		OSCPackets.Buffer.Enqueue(MoveTemp(packet));
+	}
+
+	if (!FPlatformAtomics::InterlockedCompareExchange(&OSCBufferLock, OSC_BUFFER_LOCKED, OSC_BUFFER_FREE) &&
+		!OSCPackets.Buffer.IsEmpty())
+	{
+		check(OSCBufferLock == OSC_BUFFER_LOCKED);
+		FSimpleDelegateGraphTask::CreateAndDispatchWhenReady(
+			FSimpleDelegateGraphTask::FDelegate::CreateLambda([this]() {
+			check(OSCBufferLock == OSC_BUFFER_LOCKED);
+			FPlatformAtomics::InterlockedCompareExchange(&OSCBufferLock, OSC_BUFFER_FREE, OSC_BUFFER_LOCKED);
+				
+				TSharedPtr<FOSCPacket> packet;
+				while (OSCPackets.Buffer.Dequeue(packet))
+				{
+					if (packet->IsMessage())
+					{
+						FOSCMessage message;
+						message.SetPacket(StaticCastSharedPtr<FOSCMessagePacket>(packet));
+						OnOscReceived.Broadcast(message.GetPacket()->GetAddress(), message);
+					}
+					else if (packet->IsBundle())
+					{
+						FOSCBundle bundle;
+						bundle.SetPacket(StaticCastSharedPtr<FOSCBundlePacket>(packet));
+						OnOscBundleReceived.Broadcast(bundle);
+					}
+				}
+				OSCPackets.Buffer.Empty();
+			}),
+			TStatId(),
+			nullptr,
+			ENamedThreads::GameThread);
+	}
+}
+
+void UOSCServer::Stop()
+{
+	delete SocketReceiver;
+	SocketReceiver = nullptr;
+
+	if (Socket)
+	{
+		ISocketSubsystem::Get(PLATFORM_SOCKETSUBSYSTEM)->DestroySocket(Socket);
+		Socket = nullptr;
+	}
+}
+
+void UOSCServer::BeginDestroy()
+{
+	Stop();
+	Super::BeginDestroy();
+}
diff --git a/Engine/Plugins/Runtime/OSC/Source/OSC/Private/OSCStream.cpp b/Engine/Plugins/Runtime/OSC/Source/OSC/Private/OSCStream.cpp
new file mode 100644
--- /dev/null
+++ b/Engine/Plugins/Runtime/OSC/Source/OSC/Private/OSCStream.cpp
@@ -0,0 +1,380 @@
+// Copyright 1998-2019 Epic Games, Inc. All Rights Reserved.
+#include "OSCStream.h"
+
+FOSCStream::FOSCStream() : Data(nullptr), Size(0) { Position = 0; }
+FOSCStream::~FOSCStream() {}
+FOSCStream::FOSCStream(char* Data, int Size) 
+{ 
+	this->Data = Data;
+	this->Size = Size;
+	Position = 0; 
+}
+
+char FOSCStream::ReadChar()
+{
+	char temp;
+	if (Read(&temp, 1) == 1)
+	{
+		return temp;
+	}
+
+	return 0;
+}
+
+void FOSCStream::WriteChar(char Char)
+{
+	Write(&Char, 1);
+}
+
+int32 FOSCStream::ReadInt32()
+{
+	char temp[4];
+	if (Read(temp, 4) == 4)
+	{
+#if PLATFORM_LITTLE_ENDIAN
+		union {
+			int32 i;
+			char c[4];
+		} u;
+
+		u.c[0] = temp[3];
+		u.c[1] = temp[2];
+		u.c[2] = temp[1];
+		u.c[3] = temp[0];
+
+		return u.i;
+#else
+		return *(int32*)temp;
+#endif
+	}
+
+	return 0;
+}
+
+void FOSCStream::WriteInt32(int32 Value)
+{
+	char temp[4];
+
+#ifdef PLATFORM_LITTLE_ENDIAN
+	union {
+		int32 i;
+		char c[4];
+	} u;
+
+	u.i = Value;
+
+	temp[3] = u.c[0];
+	temp[2] = u.c[1];
+	temp[1] = u.c[2];
+	temp[0] = u.c[3];
+#else
+	*reinterpret_cast<int32*>(temp) = Value;
+#endif
+	Write(temp, 4);
+}
+
+double FOSCStream::ReadDouble()
+{
+	char temp[8];
+	if (Read(temp, 8) == 8)
+	{
+#if PLATFORM_LITTLE_ENDIAN
+		union {
+			double d;
+			char c[8];
+		} u;
+
+		u.c[0] = temp[7];
+		u.c[1] = temp[6];
+		u.c[2] = temp[5];
+		u.c[3] = temp[4];
+		u.c[4] = temp[3];
+		u.c[5] = temp[2];
+		u.c[6] = temp[1];
+		u.c[7] = temp[0];
+
+		return u.d;
+#else
+		return *(double*)temp;
+#endif
+
+	}
+
+	return 0;
+}
+
+void FOSCStream::WriteDouble(uint64 Value)
+{
+	char temp[8];
+
+#ifdef PLATFORM_LITTLE_ENDIAN
+	union {
+		double i;
+		char c[8];
+	} u;
+
+	u.i = Value;
+
+	temp[7] = u.c[0];
+	temp[6] = u.c[1];
+	temp[5] = u.c[2];
+	temp[4] = u.c[3];
+	temp[3] = u.c[4];
+	temp[2] = u.c[5];
+	temp[1] = u.c[6];
+	temp[0] = u.c[7];
+#else
+	*reinterpret_cast<double*>(temp) = Value;
+#endif
+	Write(temp, 8);
+}
+
+
+int64 FOSCStream::ReadInt64()
+{
+	char temp[8];
+	if (Read(temp, 8) == 8)
+	{
+#if PLATFORM_LITTLE_ENDIAN
+		union {
+			int64 i;
+			char c[8];
+		} u;
+
+		u.c[0] = temp[7];
+		u.c[1] = temp[6];
+		u.c[2] = temp[5];
+		u.c[3] = temp[4];
+		u.c[4] = temp[3];
+		u.c[5] = temp[2];
+		u.c[6] = temp[1];
+		u.c[7] = temp[0];
+
+		return u.i;
+#else
+		return *(int64*)temp;
+#endif
+	}
+
+	return 0;
+}
+
+void FOSCStream::WriteInt64(int64 Value)
+{
+	char temp[8];
+
+#ifdef PLATFORM_LITTLE_ENDIAN
+	union {
+		int64 i;
+		char c[8];
+	} u;
+
+	u.i = Value;
+
+	temp[7] = u.c[0];
+	temp[6] = u.c[1];
+	temp[5] = u.c[2];
+	temp[4] = u.c[3];
+	temp[3] = u.c[4];
+	temp[2] = u.c[5];
+	temp[1] = u.c[6];
+	temp[0] = u.c[7];
+#else
+	*reinterpret_cast<int64*>(temp) = Value;
+#endif
+	Write(temp, 8);
+}
+
+uint64 FOSCStream::ReadUInt64()
+{
+	char temp[8];
+	if (Read(temp, 8) == 8)
+	{
+#if PLATFORM_LITTLE_ENDIAN
+		union {
+			uint64 i;
+			char c[8];
+		} u;
+
+		u.c[0] = temp[7];
+		u.c[1] = temp[6];
+		u.c[2] = temp[5];
+		u.c[3] = temp[4];
+		u.c[4] = temp[3];
+		u.c[5] = temp[2];
+		u.c[6] = temp[1];
+		u.c[7] = temp[0];
+
+		return u.i;
+#else
+		return *(uint64*)temp;
+#endif
+
+	}
+
+	return 0;
+}
+
+void FOSCStream::WriteUInt64(uint64 Value)
+{
+	char temp[8];
+
+#ifdef PLATFORM_LITTLE_ENDIAN
+	union {
+		uint64 i;
+		char c[8];
+	} u;
+
+	u.i = Value;
+
+	temp[7] = u.c[0];
+	temp[6] = u.c[1];
+	temp[5] = u.c[2];
+	temp[4] = u.c[3];
+	temp[3] = u.c[4];
+	temp[2] = u.c[5];
+	temp[1] = u.c[6];
+	temp[0] = u.c[7];
+#else
+	*reinterpret_cast<uint64*>(temp) = Value;
+#endif
+	Write(temp, 8);
+}
+
+float FOSCStream::ReadFloat()
+{
+	char temp[4];
+	if (Read(temp, 4) == 4)
+	{
+#if PLATFORM_LITTLE_ENDIAN
+		union {
+			float f;
+			char c[4];
+		} u;
+
+		u.c[0] = temp[3];
+		u.c[1] = temp[2];
+		u.c[2] = temp[1];
+		u.c[3] = temp[0];
+
+		return u.f;
+#else
+		return *(float*)temp;
+#endif
+	}
+
+	return 0.0f;
+}
+
+void FOSCStream::WriteFloat(float Value)
+{
+	char temp[4];
+
+#ifdef PLATFORM_LITTLE_ENDIAN
+	union {
+		float f;
+		char c[4];
+	} u;
+
+	u.f = Value;
+
+	temp[3] = u.c[0];
+	temp[2] = u.c[1];
+	temp[1] = u.c[2];
+	temp[0] = u.c[3];
+#else
+	*reinterpret_cast<float*>(temp) = Value;
+#endif
+	Write(temp, 4);
+}
+
+FString FOSCStream::ReadString()
+{
+	int32 count = 0;
+	for (int32 index = Position; Data[index] != 0; index++)
+	{
+		count++;
+	}
+
+	FString s(count, &Data[Position]);
+
+	Position += count + 1; // increment Position
+	Position = ((Position + 3) / 4) * 4; // padded
+
+	return s;
+}
+
+void FOSCStream::WriteString(FString string)
+{
+	TArray<TCHAR> s = string.GetCharArray();
+
+	const int32 count = s.Num();
+	for (int32 index = 0; index < count; index++)
+	{
+		Data[Position + index] = s[index];
+	}
+
+	Position += count;
+
+	const int32 numPaddingZeros = ((count + 3) / 4) * 4; // padded
+	for (int32 i = 0; i < numPaddingZeros - count; i++)
+	{
+		WriteChar('\0');
+	}
+}
+
+TArray<uint8> FOSCStream::ReadBlob()
+{
+	TArray<uint8> Blob;
+
+	int blobsize = ReadInt32();
+
+	for (int i = 0; i < blobsize; i++)
+	{
+		Blob.Add(ReadChar());
+	}
+
+	Position = ((Position + 3) / 4) * 4; // padded
+
+	return Blob;
+}
+
+void FOSCStream::WriteBlob(TArray<uint8>& Blob)
+{
+	// Write Blob size
+	WriteInt32(Blob.Num());
+
+	for (int i = 0; i < Blob.Num(); i++)
+	{
+		Write(&Blob[i], 4);
+	}
+}
+
+int FOSCStream::Read(void* Buffer, int ToRead)
+{
+	if (ToRead <= 0 || Position >= Size)
+		return 0;
+
+	int num = FMath::Min<int>(ToRead, Size - Position);
+	if (num > 0)
+	{
+		memcpy(Buffer, Data + Position, num);
+
+		Position += num;
+	}
+
+	return num;
+}
+
+int FOSCStream::Write(void* Buffer, int ToWrite)
+{
+	if (ToWrite <= 0)
+		return 0;
+
+	memcpy(Data + Position, Buffer, ToWrite);
+
+	Position += ToWrite;
+
+	return ToWrite;
+}
+
diff --git a/Engine/Plugins/Runtime/OSC/Source/OSC/Public/OSCManager.cpp b/Engine/Plugins/Runtime/OSC/Source/OSC/Public/OSCManager.cpp
new file mode 100644
--- /dev/null
+++ b/Engine/Plugins/Runtime/OSC/Source/OSC/Public/OSCManager.cpp
@@ -0,0 +1,284 @@
+// Copyright 1998-2019 Epic Games, Inc. All Rights Reserved.
+#include "OSCManager.h"
+#include "OSCMessage.h"
+#include "OSCBundle.h"
+#include "OSCLog.h"
+
+UOSCServer* UOSCManager::CreateOSCServer(FString ReceiveAddress, int32 Port)
+{
+	UOSCServer* NewOSCServer = NewObject<UOSCServer>();
+
+	FIPv4Address ipAddr;
+	if (FIPv4Address::Parse(ReceiveAddress, ipAddr))
+	{
+		NewOSCServer->Listen(ipAddr, Port, false);
+	}
+
+	return NewOSCServer;
+}
+
+UOSCClient* UOSCManager::CreateOSCClient(FString SendAddress, int32 Port)
+{
+	UOSCClient* NewOSCClient = NewObject<UOSCClient>();
+
+	FIPv4Address ipAddr;
+	if (FIPv4Address::Parse(SendAddress, ipAddr))
+	{
+		NewOSCClient->SetTarget(ipAddr, Port);
+	}
+
+	return NewOSCClient;
+}
+
+void UOSCManager::CreateOSCMessage(UPARAM(ref) FOSCMessage& Message)
+{
+	Message.SetPacket(MakeShareable(new FOSCMessagePacket()));
+}
+
+void UOSCManager::CreateOSCBundle(UPARAM(ref) FOSCBundle& Bundle)
+{
+	Bundle.SetPacket(MakeShareable(new FOSCBundlePacket()));
+}
+
+void UOSCManager::ClearMessage(UPARAM(ref) FOSCMessage& Message)
+{
+	if (Message.IsValid())
+	{
+		Message.GetPacket()->Clear();
+	}
+}
+
+void UOSCManager::ClearBundle(UPARAM(ref) FOSCBundle& Bundle)
+{
+	if (Bundle.IsValid())
+	{
+		Bundle.GetPacket()->Clear();
+	}
+}
+
+void UOSCManager::AddMessageToBundle(const FOSCMessage& Message, UPARAM(ref) FOSCBundle& Bundle)
+{
+	if (Message.IsValid() && Bundle.IsValid())
+	{
+		Bundle.GetPacket()->AddPacket(Message.GetPacket());
+	}
+}
+
+void UOSCManager::AddBundleToBundle(const FOSCBundle& inBundle, UPARAM(ref) FOSCBundle& outBundle)
+{
+	if (inBundle.IsValid() && outBundle.IsValid())
+	{
+		outBundle.GetPacket()->AddPacket(inBundle.GetPacket());
+	}
+}
+
+void UOSCManager::SetAddress(UPARAM(ref) FOSCMessage& Message, const FString& Address)
+{
+	Message.GetOrCreatePacket()->SetAddress(Address);
+}
+
+void UOSCManager::AddFloat(UPARAM(ref) FOSCMessage& Message, float Value)
+{
+	Message.GetOrCreatePacket()->AddArgument(FOSCType(Value));
+}
+
+void UOSCManager::AddInt32(UPARAM(ref) FOSCMessage& Message, int32 Value)
+{
+	Message.GetOrCreatePacket()->AddArgument(FOSCType(Value));
+}
+
+void UOSCManager::AddInt64(UPARAM(ref) FOSCMessage& Message, int64 Value)
+{
+	Message.GetOrCreatePacket()->AddArgument(FOSCType(Value));
+}
+
+void UOSCManager::AddString(UPARAM(ref) FOSCMessage& Message, FString Value)
+{
+	Message.GetOrCreatePacket()->AddArgument(FOSCType(Value));
+}
+
+void UOSCManager::AddBlob(UPARAM(ref) FOSCMessage& Message, TArray<uint8>& Value)
+{
+	Message.GetOrCreatePacket()->AddArgument(FOSCType(Value));
+}
+
+void UOSCManager::AddBool(UPARAM(ref) FOSCMessage& Message, bool Value)
+{
+	Message.GetOrCreatePacket()->AddArgument(FOSCType(Value));
+}
+
+void UOSCManager::GetMessagesFromBundle(const FOSCBundle& Bundle, TArray<FOSCMessage>& Messages)
+{
+	if (Bundle.IsValid())
+	{
+		const TSharedPtr<FOSCBundlePacket> bundlePacket = Bundle.GetPacket();
+
+		const int32 num = bundlePacket->GetNumPackets();
+		for (int32 i = 0; i < num; i++)
+		{
+			TSharedPtr<FOSCPacket> packet = bundlePacket->GetPacket(i);
+
+			if (packet->IsMessage())
+			{
+				FOSCMessage message;
+
+				message.SetPacket(StaticCastSharedPtr<FOSCMessagePacket>(packet));
+
+				Messages.Add(message);
+			}
+		}
+	}
+}
+
+void UOSCManager::GetFloat(const FOSCMessage& Message, const int index, float& Value)
+{
+	if (Message.IsValid())
+	{
+		const FOSCType& t = Message.GetPacket()->GetArgument(index);
+
+		if (t.IsFloat())
+		{
+			Value = t.GetFloat();
+		}
+	}
+}
+
+void UOSCManager::GetAllFloat(const FOSCMessage& Message, TArray<float>& Values)
+{
+	if (Message.IsValid())
+	{
+		for (int i = 0; i < Message.GetPacket()->GetNumArguments(); i++)
+		{
+			const FOSCType& t = Message.GetPacket()->GetArgument(i);
+			if (t.IsFloat())
+			{
+				Values.Add(t.GetFloat());
+			}
+		}
+	}
+}
+
+void UOSCManager::GetInt32(const FOSCMessage& Message, const int index, int32& Value)
+{
+	if (Message.IsValid())
+	{
+		const FOSCType& t = Message.GetPacket()->GetArgument(index);
+
+		if (t.IsInt32())
+		{
+			Value = t.GetInt32();
+		}
+	}
+}
+
+void UOSCManager::GetAllInt32(const FOSCMessage& Message, TArray<int32>& Values)
+{
+	if (Message.IsValid())
+	{
+		for (int i = 0; i < Message.GetPacket()->GetNumArguments(); i++)
+		{
+			const FOSCType& t = Message.GetPacket()->GetArgument(i);
+			if (t.IsInt32())
+			{
+				Values.Add(t.GetInt32());
+			}
+		}
+	}
+}
+
+void UOSCManager::GetInt64(const FOSCMessage& Message, const int index, int64& Value)
+{
+	if (Message.IsValid())
+	{
+		const FOSCType& t = Message.GetPacket()->GetArgument(index);
+
+		if (t.IsInt64())
+		{
+			Value = t.GetInt64();
+		}
+	}
+}
+
+void UOSCManager::GetAllInt64(const FOSCMessage& Message, TArray<int64>& Values)
+{
+	if (Message.IsValid())
+	{
+		for (int i = 0; i < Message.GetPacket()->GetNumArguments(); i++)
+		{
+			const FOSCType& t = Message.GetPacket()->GetArgument(i);
+			if (t.IsInt64())
+			{
+				Values.Add(t.GetInt64());
+			}
+		}
+	}
+}
+
+void UOSCManager::GetString(const FOSCMessage& Message, const int index, FString& Value)
+{
+	if (Message.IsValid())
+	{
+		const FOSCType& t = Message.GetPacket()->GetArgument(index);
+
+		if (t.IsString())
+		{
+			Value = t.GetString();
+		}
+	}
+}
+
+void UOSCManager::GetAllString(const FOSCMessage& Message, TArray<FString>& Values)
+{
+	if (Message.IsValid())
+	{
+		for (int i = 0; i < Message.GetPacket()->GetNumArguments(); i++)
+		{
+			const FOSCType& t = Message.GetPacket()->GetArgument(i);
+			if (t.IsString())
+			{
+				Values.Add(t.GetString());
+			}
+		}
+	}
+}
+
+void UOSCManager::GetBool(const FOSCMessage& Message, const int index, bool& Value)
+{
+	if (Message.IsValid())
+	{
+		const FOSCType& t = Message.GetPacket()->GetArgument(index);
+
+		if (t.IsBool())
+		{
+			Value = t.GetBool();
+		}
+	}
+}
+
+void UOSCManager::GetAllBool(const FOSCMessage& Message, TArray<bool>& Values)
+{
+	if (Message.IsValid())
+	{
+		for (int i = 0; i < Message.GetPacket()->GetNumArguments(); i++)
+		{
+			const FOSCType& t = Message.GetPacket()->GetArgument(i);
+			if (t.IsBool())
+			{
+				Values.Add(t.GetBool());
+			}
+		}
+	}
+}
+
+void UOSCManager::GetBlob(const FOSCMessage& Message, const int index, TArray<uint8>& Value)
+{
+	if (Message.IsValid())
+	{
+		const FOSCType& t = Message.GetPacket()->GetArgument(index);
+
+		if (t.IsBlob())
+		{
+			Value = t.GetBlob();
+		}
+	}
+}
diff --git a/Engine/Plugins/Runtime/OSC/Source/OSC/Public/OSCManager.h b/Engine/Plugins/Runtime/OSC/Source/OSC/Public/OSCManager.h
new file mode 100644
--- /dev/null
+++ b/Engine/Plugins/Runtime/OSC/Source/OSC/Public/OSCManager.h
@@ -0,0 +1,101 @@
+// Copyright 1998-2019 Epic Games, Inc. All Rights Reserved.
+#pragma once
+
+#include "CoreMinimal.h"
+#include "UObject/ObjectMacros.h"
+#include "Kismet/BlueprintFunctionLibrary.h"
+#include "OSCServer.h"
+#include "OSCClient.h"
+#include "OSCManager.generated.h"
+
+UCLASS()
+class OSC_API UOSCManager : public UBlueprintFunctionLibrary
+{
+	GENERATED_BODY()
+
+public:
+	//////////////////////////////////////////////////////////////////////////
+	// Audio|OSC functions
+	UFUNCTION(BlueprintCallable, Category = "Audio|OSC")
+	static class UOSCServer* CreateOSCServer(FString ReceiveAddress, int32 Port);
+
+	UFUNCTION(BlueprintCallable, Category = "Audio|OSC")
+	static class UOSCClient* CreateOSCClient(FString SendAddress, int32 Port);
+
+	UFUNCTION(BlueprintCallable, Category = "Audio|OSC")
+	static void CreateOSCMessage(UPARAM(ref) FOSCMessage& Message);
+
+	UFUNCTION(BlueprintCallable, Category = "Audio|OSC")
+	static void CreateOSCBundle(UPARAM(ref) FOSCBundle& Bundle);
+
+	UFUNCTION(BlueprintCallable, Category = "Audio|OSC")
+	static void AddMessageToBundle(const FOSCMessage& Message, UPARAM(ref) FOSCBundle& Bundle);
+
+	UFUNCTION(BlueprintCallable, Category = "Audio|OSC")
+	static void AddBundleToBundle(const FOSCBundle& inBundle, UPARAM(ref) FOSCBundle& outBundle);
+
+	UFUNCTION(BlueprintCallable, Category = "Audio|OSC")
+	static void GetMessagesFromBundle(const FOSCBundle& Bundle, TArray<FOSCMessage>& Messages);
+
+	UFUNCTION(BlueprintCallable, Category = "Audio|OSC")
+	static void ClearMessage(UPARAM(ref) FOSCMessage& Message);
+
+	UFUNCTION(BlueprintCallable, Category = "Audio|OSC")
+	static void ClearBundle(UPARAM(ref) FOSCBundle& Bundle);
+
+	UFUNCTION(BlueprintCallable, Category = "Audio|OSC")
+	static void SetAddress(UPARAM(ref) FOSCMessage& Message, const FString& Address);
+
+	UFUNCTION(BlueprintCallable, Category = "Audio|OSC")
+	static void AddFloat(UPARAM(ref) FOSCMessage& Message, float Value);
+
+	UFUNCTION(BlueprintCallable, Category = "Audio|OSC")
+	static void AddInt32(UPARAM(ref) FOSCMessage& Message, int32 Value);
+
+	UFUNCTION(BlueprintCallable, Category = "Audio|OSC")
+	static void AddInt64(UPARAM(ref) FOSCMessage& Message, int64 Value);
+
+	UFUNCTION(BlueprintCallable, Category = "Audio|OSC")
+	static void AddString(UPARAM(ref) FOSCMessage& Message, FString Value);
+
+	UFUNCTION(BlueprintCallable, Category = "Audio|OSC")
+	static void AddBlob(UPARAM(ref) FOSCMessage& Message, TArray<uint8>& Value);
+
+	UFUNCTION(BlueprintCallable, Category = "Audio|OSC")
+	static void AddBool(UPARAM(ref) FOSCMessage& Message, bool Value);
+
+	//////////////////////////////////////////////////////////////////////////
+	// Audio|OSC functions
+	UFUNCTION(BlueprintCallable, Category = "Audio|OSC")
+	static void GetFloat(const FOSCMessage& Message, const int index, float& Value);
+
+	UFUNCTION(BlueprintCallable, Category = "Audio|OSC")
+	static void GetAllFloat(const FOSCMessage& Message, TArray<float>& Values);
+
+	UFUNCTION(BlueprintCallable, Category = "Audio|OSC")
+	static void GetInt32(const FOSCMessage& Message, const int index, int32& Value);
+
+	UFUNCTION(BlueprintCallable, Category = "Audio|OSC")
+	static void GetAllInt32(const FOSCMessage& Message, TArray<int32>& Values);
+
+	UFUNCTION(BlueprintCallable, Category = "Audio|OSC")
+	static void GetInt64(const FOSCMessage& Message, const int index, int64& Value);
+
+	UFUNCTION(BlueprintCallable, Category = "Audio|OSC")
+	static void GetAllInt64(const FOSCMessage& Message, TArray<int64>& Values);
+
+	UFUNCTION(BlueprintCallable, Category = "Audio|OSC")
+	static void GetString(const FOSCMessage& Message, const int index, FString& Value);
+
+	UFUNCTION(BlueprintCallable, Category = "Audio|OSC")
+	static void GetAllString(const FOSCMessage& Message, TArray<FString>& Values);
+
+	UFUNCTION(BlueprintCallable, Category = "Audio|OSC")
+	static void GetBool(const FOSCMessage& Message, const int index, bool& Value);
+
+	UFUNCTION(BlueprintCallable, Category = "Audio|OSC")
+	static void GetAllBool(const FOSCMessage& Message, TArray<bool>& Values);
+
+	UFUNCTION(BlueprintCallable, Category = "Audio|OSC")
+	static void GetBlob(const FOSCMessage& Message, const int index, TArray<uint8>& Value);
+};
diff --git a/Engine/Plugins/Runtime/OSC/Source/OSC/Public/OSCModule.cpp b/Engine/Plugins/Runtime/OSC/Source/OSC/Public/OSCModule.cpp
new file mode 100644
--- /dev/null
+++ b/Engine/Plugins/Runtime/OSC/Source/OSC/Public/OSCModule.cpp
@@ -0,0 +1,32 @@
+// Copyright 1998-2019 Epic Games, Inc. All Rights Reserved.
+#include "CoreMinimal.h"
+#include "Modules/ModuleInterface.h"
+#include "Modules/ModuleManager.h"
+#include "OSCLog.h"
+
+DEFINE_LOG_CATEGORY( OSCLog );
+
+/**
+ * The public interface to this module
+ */
+class FOSCModule : public IModuleInterface
+{
+public:
+	virtual void StartupModule() override;
+	virtual void ShutdownModule() override;
+};
+
+IMPLEMENT_MODULE( FOSCModule, OSC )
+
+void FOSCModule::StartupModule()
+{
+	if (!FModuleManager::Get().LoadModule(TEXT("Networking")))
+	{
+		UE_LOG(OSCLog, Error, TEXT("Required module Networking failed to load"));
+		return;
+	}
+}
+
+void FOSCModule::ShutdownModule()
+{
+}
diff --git a/Engine/Plugins/Runtime/Synthesis/Source/Synthesis/Classes/SourceEffects/SourceEffectMidSideSpreader.h b/Engine/Plugins/Runtime/Synthesis/Source/Synthesis/Classes/SourceEffects/SourceEffectMidSideSpreader.h
new file mode 100644
--- /dev/null
+++ b/Engine/Plugins/Runtime/Synthesis/Source/Synthesis/Classes/SourceEffects/SourceEffectMidSideSpreader.h
@@ -0,0 +1,110 @@
+// Copyright 1998-2019 Epic Games, Inc. All Rights Reserved.
+
+#pragma once
+
+#include "CoreMinimal.h"
+#include "Sound/SoundEffectSource.h"
+#include "UObject/ObjectMacros.h"
+#include "DSP/Dsp.h"
+#include "SourceEffectMidSideSpreader.generated.h"
+
+// Stereo channel mode
+UENUM(BlueprintType)
+enum class EStereoChannelMode : uint8
+{
+		MidSide,
+		LeftRight,
+		count UMETA(Hidden)
+};
+
+// ========================================================================
+// FSourceEffectMidSideSpreaderSettings
+// This is the source effect's setting struct. 
+// ========================================================================
+
+USTRUCT(BlueprintType)
+struct SYNTHESIS_API FSourceEffectMidSideSpreaderSettings
+{
+	GENERATED_USTRUCT_BODY()
+
+	// Indicate the channel mode of the input signal
+	UPROPERTY(EditAnywhere, BlueprintReadWrite, Category = "SourceEffect|Preset")
+	EStereoChannelMode InputMode;
+
+	// Amount of Mid/Side Spread. 0.0 is no spread, 1.0 is full wide. 
+	UPROPERTY(EditAnywhere, BlueprintReadWrite, Category = "SourceEffect|Preset", meta = (ClampMin = "0.0", ClampMax = "1.0", UIMin = "0.0", UIMax = "1.0"))
+	float SpreadAmount;
+
+	// Indicate the channel mode of the output signal
+	UPROPERTY(EditAnywhere, BlueprintReadWrite, Category = "SourceEffect|Preset")
+	EStereoChannelMode OutputMode;
+
+	// Indicate whether an equal power relationship between the mid and side channels should be maintained
+	UPROPERTY(EditAnywhere, BlueprintReadWrite, Category = "SourceEffect|Preset")
+	bool bEqualPower;
+
+	FSourceEffectMidSideSpreaderSettings()
+		: InputMode(EStereoChannelMode::LeftRight)
+		, SpreadAmount(0.5f)
+		, OutputMode(EStereoChannelMode::LeftRight)
+		, bEqualPower(false)
+	{}
+};
+
+// ========================================================================
+// FSourceEffectMidSideSpreader
+// This is the instance of the source effect. Performs DSP calculations.
+// ========================================================================
+
+class SYNTHESIS_API FSourceEffectMidSideSpreader : public FSoundEffectSource
+{
+public:
+	FSourceEffectMidSideSpreader();
+
+	// Called on an audio effect at initialization on main thread before audio processing begins.
+	virtual void Init(const FSoundEffectSourceInitData& InInitData) override;
+	
+	// Called when an audio effect preset is changed
+	virtual void OnPresetChanged() override;
+
+	// Process the input block of audio. Called on audio thread.
+	virtual void ProcessAudio(const FSoundEffectSourceInputData& InData, float* OutAudioBufferData) override;
+
+protected:
+
+	float MidScale;
+	float SideScale;
+
+	int32 NumChannels;
+
+	FSourceEffectMidSideSpreaderSettings SpreaderSettings;
+
+};
+
+// ========================================================================
+// USourceEffectMidSideSpreaderPreset
+// This code exposes your preset settings and effect class to the editor.
+// And allows for a handle to setting/updating effect settings dynamically.
+// ========================================================================
+
+UCLASS(ClassGroup = AudioSourceEffect, meta = (BlueprintSpawnableComponent))
+class SYNTHESIS_API USourceEffectMidSideSpreaderPreset : public USoundEffectSourcePreset
+{
+	GENERATED_BODY()
+
+public:
+	// Macro which declares and implements useful functions.
+	EFFECT_PRESET_METHODS(SourceEffectMidSideSpreader)
+
+	// Allows you to customize the color of the preset in the editor.
+	virtual FColor GetPresetColor() const override { return FColor(126.0f, 180.0f, 255.0f); }
+
+	// Change settings of your effect from blueprints. Will broadcast changes to active instances.
+	UFUNCTION(BlueprintCallable, Category = "Audio|Effects")
+	void SetSettings(const FSourceEffectMidSideSpreaderSettings& InSettings);
+	
+	// The copy of the settings struct. Can't be written to in BP, but can be read.
+	// Note that the value read in BP is the serialized settings, will not reflect dynamic changes made in BP.
+	UPROPERTY(EditAnywhere, BlueprintReadOnly, Category = "Audio|Effects")
+	FSourceEffectMidSideSpreaderSettings Settings;
+};
diff --git a/Engine/Plugins/Runtime/Synthesis/Source/Synthesis/Classes/SourceEffects/SourceEffectPanner.h b/Engine/Plugins/Runtime/Synthesis/Source/Synthesis/Classes/SourceEffects/SourceEffectPanner.h
--- a/Engine/Plugins/Runtime/Synthesis/Source/Synthesis/Classes/SourceEffects/SourceEffectPanner.h
+++ b/Engine/Plugins/Runtime/Synthesis/Source/Synthesis/Classes/SourceEffects/SourceEffectPanner.h
@@ -11,12 +11,17 @@ struct SYNTHESIS_API FSourceEffectPannerSettings
 {
 	GENERATED_USTRUCT_BODY()
 
+	// The spread of the source. 1.0 means left only in left channel, right only in right; 0.0 means both mixed, -1.0 means right and left channels are inverted.
+	UPROPERTY(EditAnywhere, BlueprintReadWrite, Category = "SourceEffect|Preset", meta = (ClampMin = "-1.0", ClampMax = "1.0", UIMin = "-1.0", UIMax = "1.0"))
+	float Spread;
+
 	// The pan of the source. -1.0 means left, 0.0 means center, 1.0 means right.
 	UPROPERTY(EditAnywhere, BlueprintReadWrite, Category = "SourceEffect|Preset", meta = (ClampMin = "-1.0", ClampMax = "1.0", UIMin = "-1.0", UIMax = "1.0"))
 	float Pan;
 
 	FSourceEffectPannerSettings()
-		: Pan(0.0f)
+		: Spread(1.0f)
+		, Pan(0.0f)
 	{}
 };
 
@@ -25,16 +30,28 @@ class SYNTHESIS_API FSourceEffectPanner : public FSoundEffectSource
 public:
 	// Called on an audio effect at initialization on main thread before audio processing begins.
 	virtual void Init(const FSoundEffectSourceInitData& InitData) override;
-	
+
 	// Called when an audio effect preset is changed
 	virtual void OnPresetChanged() override;
 
 	// Process the input block of audio. Called on audio render thread.
 	virtual void ProcessAudio(const FSoundEffectSourceInputData& InData, float* OutAudioBufferData) override;
 
+	FSourceEffectPanner()
+	{
+		for (int32 i = 0; i < 2; ++i)
+		{
+			SpreadGains[i] = 0.0f;
+			PanGains[i] = 0.0f;
+		}
+		NumChannels = 0;
+	}
+
 protected:
+
 	// The pan value of the source effect
-	float PanValue;
+	float SpreadGains[2];
+	float PanGains[2];
 	int32 NumChannels;
 };
 
diff --git a/Engine/Plugins/Runtime/Synthesis/Source/Synthesis/Classes/SynthComponents/SynthComponentMonoWaveTable.h b/Engine/Plugins/Runtime/Synthesis/Source/Synthesis/Classes/SynthComponents/SynthComponentMonoWaveTable.h
new file mode 100644
--- /dev/null
+++ b/Engine/Plugins/Runtime/Synthesis/Source/Synthesis/Classes/SynthComponents/SynthComponentMonoWaveTable.h
@@ -0,0 +1,415 @@
+// Copyright 1998-2019 Epic Games, Inc. All Rights Reserved.
+
+#pragma once
+
+#include "CoreMinimal.h"
+#include "Components/SynthComponent.h"
+#include "DSP/Osc.h"
+#include "DSP/LFO.h"
+#include "Components/AudioComponent.h"
+#include "UObject/ObjectMacros.h"
+#include "UObject/Object.h"
+#include "UObject/Class.h"
+#include "Engine/EngineTypes.h"
+#include "Sound/SoundWaveProcedural.h"
+#include "Engine/DataTable.h"
+#include "Components/SynthComponent.h"
+#include "DSP/Osc.h"
+#include "Kismet/BlueprintFunctionLibrary.h"
+#include "MonoWaveTable.h"
+#include "EpicSynth1Types.h"
+#include "SynthesisModule.h"
+#include "SynthComponentMonoWaveTable.generated.h"
+
+// ========================================================================
+// USynthComponentMonoWaveTable
+// Synth component class which implements USynthComponent
+// This is a simple hello-world type example which generates a sine-wave
+// tone using a DSP oscillator class and implements a single function to set
+// the frequency. 
+// ========================================================================
+
+class USynthComponentMonoWaveTable;
+
+// struct passed from Mono WaveTable Asset
+// to USynthObject
+struct AssetChangeInfo
+{
+	uint8 bNeedsFullRebuild : 1;
+	int32 CurveThatWasAltered{ -1 };
+
+	void FlagCurveAsAltered(const int32 CurveIndex)
+	{
+		if (CurveThatWasAltered == -1)
+		{
+			CurveThatWasAltered = CurveIndex;
+		}
+		else
+		{
+			bNeedsFullRebuild = true;
+		}
+	}
+
+	AssetChangeInfo() : bNeedsFullRebuild(false) { }
+}; // struct AssetChaneInfo
+
+// UStruct Mono Wave Table Synth Preset
+UCLASS(ClassGroup = Synth, meta = (BlueprintSpawnableComponent))
+class SYNTHESIS_API UMonoWaveTableSynthPreset : public UObject
+{
+	GENERATED_BODY()
+
+public:
+
+	// ctor
+	UMonoWaveTableSynthPreset();
+
+	// Name the preset	
+	UPROPERTY(EditAnywhere, BlueprintReadWrite, Category = "Synth|Preset")
+	FString PresetName;
+
+	// Lock wavetables to evenly spaced keyframes that can be edited vertically only (will re-sample)
+	UPROPERTY(EditAnywhere, BlueprintReadOnly, Category = "Synth|Preset", Meta = (InlineEditConditionToggle))
+	uint8 bLockKeyframesToGridBool : 1;
+
+	// How many evenly-spaced keyframes to use when LockKeyframesToGrid is true
+	UPROPERTY(EditAnywhere, BlueprintReadOnly, Category = "Synth|Preset", Meta = (EditCondition = "bLockKeyframesToGridBool", ClampMin = "3", UIMin = "3", ClampMax = "256", UIMax = "256"))
+	int32 LockKeyframesToGrid;
+
+	// How many samples will be taken of the curve from time = [0.0, 1.0]
+	UPROPERTY(EditAnywhere, BlueprintReadOnly, Category = "Synth|Preset", Meta = (ClampMin = "3", UIMin = "3", ClampMax = "4096", UIMax = "4096"))
+	int32 WaveTableResolution;
+
+	// Wave Table Editor
+	UPROPERTY(EditAnywhere, NonTransactional, BlueprintReadOnly, Category = "Synth|Preset")
+	TArray<FRuntimeFloatCurve> WaveTable;
+
+	// Normalize the WaveTable data? False will allow clipping, True will normalize the tables when sent to the synth for rendering
+	UPROPERTY(EditAnywhere, BlueprintReadWrite, Category = "Synth|Preset")
+	uint8 bNormalizeWaveTables : 1;
+
+	// function to register lambdas to call when a property is changed
+	void RegisterWTComponentCallback(uint32 ID, TFunction<void(const AssetChangeInfo& ChangeInfo)> Callback);
+
+	// function to UNregister same lambdas
+	void UnRegisterWTSynthComponentCallback(const uint32 ID);
+
+protected:
+#if WITH_EDITOR
+	// Override on Post property change thing here (loop and call lambdas)
+	virtual void PostEditChangeChainProperty(struct FPropertyChangedChainEvent&) override;
+#endif // WITH_EDITOR
+
+	void EditChangeInternal();
+
+	void SampleAllToGrid(uint32 InGridsize);
+	
+	void SampleToGrid(uint32 InGridSize, uint32 InTableIndex);
+
+	// Since wavetable synthesis sounds good when there are subtle changes between curves, 
+	// this helps work flow when adding a new curve by making it a duplicate of a the curve before it
+	void DuplicateCurveToEnd();
+
+	// saves data from the asset to see what changed when
+	void CacheAssetData();
+
+	// Compares Underlying rich curves to see if the Index(th) curve was changed in the editor
+	bool IsCachedTableEntryStillValid(Audio::DefaultWaveTableIndexType Index);
+
+	// Map of registered TFunctions
+	TMap<uint32, TFunction<void(const AssetChangeInfo&)>> PropertyChangedCallbacks;
+
+	TArray<float> CurveBiDirTangents;
+
+	// Default curve to use for work flow QOL (Never expose empty/silent curve)
+	FRuntimeFloatCurve DefaultCurve;
+
+	int32 CachedGridSize;
+
+	// cached asset data (before last edit)
+	int8 bWasLockedToGrid : 1;
+	int8 bWasSingleTable : 1;
+	int32 CachedTableResolution;
+	TArray<FRuntimeFloatCurve> CachedWaveTable;
+	uint8 bCachedNormalizationSetting : 1;
+
+
+	// let the USynth see into asset
+	friend class USynthComponentMonoWaveTable;
+};
+
+DECLARE_DYNAMIC_MULTICAST_DELEGATE_OneParam(FOnTableAltered, int32, TableIndex);
+DECLARE_DYNAMIC_MULTICAST_DELEGATE(FNumTablesChanged);
+
+UENUM(BlueprintType)
+enum class CurveInterpolationType : uint8
+{
+	AUTOINTERP UMETA(DisplayName = "Auto"),
+	LINEAR UMETA(DisplayName = "Linear"),
+	CONSTANT UMETA(DisplayName = "Constant")
+};
+
+
+UCLASS(ClassGroup = Synth, meta = (BlueprintSpawnableComponent))
+class SYNTHESIS_API USynthComponentMonoWaveTable : public USynthComponent
+{
+	GENERATED_BODY()
+
+	USynthComponentMonoWaveTable(const FObjectInitializer& ObjectInitializer);
+
+	// if in editor, unregister with the asset
+	~USynthComponentMonoWaveTable()
+	{
+		if (CachedPreset)
+		{
+			CachedPreset->UnRegisterWTSynthComponentCallback(GetUniqueID());
+		}
+	}
+	
+	// Called when synth is created
+	virtual bool Init(int32& InSampleRate) override;
+
+	// Called to generate more audio
+	virtual Audio::DefaultWaveTableIndexType OnGenerateAudio(float* OutAudio, int32 NumSamples) override;
+
+public:
+	/* Start BP functionality */
+
+	// Get the number of table elements from Blueprint
+	UFUNCTION(BlueprintCallable, Category = "Synth|Components|Audio")
+	int32 GetNumTableEntries();
+
+	// Starts a new note (retrigs modulators, etc.)
+	UFUNCTION(BlueprintCallable, Category = "Synth|Components|Audio")
+	void NoteOn(const float InMidiNote, const float InVelocity);
+
+	// Starts a new note (retrigs modulators, etc.)
+	UFUNCTION(BlueprintCallable, Category = "Synth|Components|Audio")
+	void NoteOff(const float InMidiNote);
+
+	// Inform the synth if the sustain pedal is pressed or not
+	UFUNCTION(BlueprintCallable, Category = "Synth|Components|Audio")
+	void SetSustainPedalState(bool InSustainPedalState);
+
+	// Sets the oscillator's frequency
+	UFUNCTION(BlueprintCallable, Category = "Synth|Components|Audio")
+	void SetFrequency(const float FrequencyHz);
+
+	// Set a frequency offset in cents (for pitch modulation such as the Pitch Bend Wheel)
+	UFUNCTION(BlueprintCallable, Category = "Synth|Components|Audio")
+	void SetFrequencyPitchBend(const float FrequencyOffsetCents);
+
+	// Set the oscillator's frequency via midi note number
+	UFUNCTION(BlueprintCallable, Category = "Synth|Components|Audio")
+	void SetFrequencyWithMidiNote(const float InMidiNote);
+
+	// Sets the wavetable position. Expects a percentage between 0.0 and 1.0
+	UFUNCTION(BLueprintCallable, Category = "Synth|Components|Audio", Meta = (ClampMin = "0.0", ClampMax = "1.0"))
+	void SetWaveTablePosition(float InPosition);
+
+	// Refresh a particular wavetable (from Game Thread data)
+	UFUNCTION(BlueprintCallable, Category = "Synth|Components|Audio")
+	void RefreshWaveTable(int32 Index);
+
+	// Refresh all wavetables (from Game Thread data)
+	UFUNCTION(BlueprintCallable, Category = "Synth|Components|Audio")
+	void RefreshAllWaveTables();
+
+	// Switch to another preset (STOPS SYNTH FROM PLAYING)
+	//UFUNCTION(BlueprintCallable, Category = "Synth|Components|Audio")
+	void SetSynthPreset(UMonoWaveTableSynthPreset* SynthPreset);
+
+	// Set frequency of LFO controlling Table Position (in Hz)
+	UFUNCTION(BlueprintCallable, Category = "Synth|Components|Audio")
+		void SetPosLfoFrequency(const float InLfoFrequency);
+
+	/* Set the Modulation depth of the Lfo controlling the Table Position around the current position value
+	   0.0 = no modulation, 1.0 = current position +/- 0.5 (Lfo + Position result will clamp [0.0, 1.0]) */
+	UFUNCTION(BlueprintCallable, Category = "Synth|Components|Audio")
+	void SetPosLfoDepth(const float InLfoDepth);
+
+	// Set the shape of the Lfo controlling the position
+	UFUNCTION(BlueprintCallable, Category = "Synth|Components|Audio")
+	void SetPosLfoType(const ESynthLFOType InLfoType);
+
+	// Set the Cut-off frequency of the low-pass filter
+	UFUNCTION(BlueprintCallable, Category = "Synth|Components|Audio")
+	void SetLowPassFilterFrequency(float InNewFrequency);
+
+	// Set the Cut-off frequency of the low-pass filter
+	UFUNCTION(BlueprintCallable, Category = "Synth|Components|Audio",  Meta = (ClampMin = "0.0", ClampMax = "10.0"))
+	void SetLowPassFilterResonance(float InNewQ);
+
+	// Set Amp envelope attack time (msec)
+	UFUNCTION(BlueprintCallable, Category = "Synth|Components|Audio")
+	void SetAmpEnvelopeAttackTime(const float InAttackTimeMsec);
+	
+	// Set Amp envelope decay time (msec)
+	UFUNCTION(BlueprintCallable, Category = "Synth|Components|Audio")
+	void SetAmpEnvelopeDecayTime(const float InDecayTimeMsec);
+	
+	// Set Amp envelope sustain gain [0.0, 1.0]
+	UFUNCTION(BlueprintCallable, Category = "Synth|Components|Audio")
+	void SetAmpEnvelopeSustainGain(const float InSustainGain);
+	
+	// Set Amp envelope release time (msec)
+	UFUNCTION(BlueprintCallable, Category = "Synth|Components|Audio")
+	void SetAmpEnvelopeReleaseTime(const float InReleaseTimeMsec);
+	
+	// Set whether or not the Amp envelope is inverted
+	UFUNCTION(BlueprintCallable, Category = "Synth|Components|Audio")
+	void SetAmpEnvelopeInvert(const bool bInInvert);
+	
+	// Set whether or not the Amp envelope's bias is inverted
+	UFUNCTION(BlueprintCallable, Category = "Synth|Components|Audio")
+	void SetAmpEnvelopeBiasInvert(const bool bInBiasInvert);
+	
+	// Set the overall depth of the Amp envelope
+	UFUNCTION(BlueprintCallable, Category = "Synth|Components|Audio")
+	void SetAmpEnvelopeDepth(const float InDepth);
+	
+	// Set the bias depth of the the Amp envelope
+	UFUNCTION(BlueprintCallable, Category = "Synth|Components|Audio")
+	void SetAmpEnvelopeBiasDepth(const float InDepth);
+
+	// Set Low-Pass Filter envelope attack time (msec)
+	UFUNCTION(BlueprintCallable, Category = "Synth|Components|Audio")
+	void SetFilterEnvelopeAttackTime(const float InAttackTimeMsec);
+	
+	// Set Low-Pass Filter envelope decay time (msec)
+	UFUNCTION(BlueprintCallable, Category = "Synth|Components|Audio")
+	void SetFilterEnvelopenDecayTime(const float InDecayTimeMsec);
+	
+	// Set Low-Pass Filter envelope sustain gain
+	UFUNCTION(BlueprintCallable, Category = "Synth|Components|Audio")
+	void SetFilterEnvelopeSustainGain(const float InSustainGain);
+	
+	// Set Low-Pass Filter envelope release time (msec)
+	UFUNCTION(BlueprintCallable, Category = "Synth|Components|Audio")
+	void SetFilterEnvelopeReleaseTime(const float InReleaseTimeMsec);
+	
+	// Set Low-Pass Filter envelope inversion
+	UFUNCTION(BlueprintCallable, Category = "Synth|Components|Audio")
+	void SetFilterEnvelopeInvert(const bool bInInvert);
+	
+	// Set Low-Pass Filter envelope bias inversion
+	UFUNCTION(BlueprintCallable, Category = "Synth|Components|Audio")
+	void SetFilterEnvelopeBiasInvert(const bool bInBiasInvert);
+	
+	// Set Low-Pass Filter envelope depth
+	UFUNCTION(BlueprintCallable, Category = "Synth|Components|Audio")
+	void SetFilterEnvelopeDepth(const float InDepth);
+	
+	// Set Low-Pass Filter envelope bias depth
+	UFUNCTION(BlueprintCallable, Category = "Synth|Components|Audio")
+	void SetFilterEnvelopeBiasDepth(const float InDepth);
+
+	// Set Position envelope attack time (msec)
+	UFUNCTION(BlueprintCallable, Category = "Synth|Components|Audio")
+	void SetPositionEnvelopeAttackTime(const float InAttackTimeMsec);
+
+	// Set Position envelope decay time (msec)
+	UFUNCTION(BlueprintCallable, Category = "Synth|Components|Audio")
+	void SetPositionEnvelopeDecayTime(const float InDecayTimeMsec);
+
+	// Set Position envelope sustain gain
+	UFUNCTION(BlueprintCallable, Category = "Synth|Components|Audio")
+	void SetPositionEnvelopeSustainGain(const float InSustainGain);
+
+	// Set Position envelope release time (msec)
+	UFUNCTION(BlueprintCallable, Category = "Synth|Components|Audio")
+	void SetPositionEnvelopeReleaseTime(const float InReleaseTimeMsec);
+
+	// Set Position envelope envelope inversion
+	UFUNCTION(BlueprintCallable, Category = "Synth|Components|Audio")
+	void SetPositionEnvelopeInvert(const bool bInInvert);
+
+	// Set Position envelope bias inversion
+	UFUNCTION(BlueprintCallable, Category = "Synth|Components|Audio")
+	void SetPositionEnvelopeBiasInvert(const bool bInBiasInvert);
+
+	// Set Position envelope envelope depth
+	UFUNCTION(BlueprintCallable, Category = "Synth|Components|Audio")
+	void SetPositionEnvelopeDepth(const float InDepth);
+
+	// Set Position envelope bias depth
+	UFUNCTION(BlueprintCallable, Category = "Synth|Components|Audio")
+	void SetPositionEnvelopeBiasDepth(const float InDepth);
+
+	// Get the number of curves in the wave table. (returns -1 if there is no asset)
+	UFUNCTION(BlueprintCallable, Category = "Synth|Components|Audio")
+	int32 GetMaxTableIndex() const { return CachedPreset? CachedPreset->WaveTable.Num() - 1 : -1; }
+
+	/* Set a Keyframe value given a Table number and Keyframe number.
+	   Returns false if the request was invalid.
+	   NewValue will be clamped from +/- 1.0 */
+	UFUNCTION(BlueprintCallable, Category = "Synth|Components|Audio")
+	bool SetCurveValue(int32 TableIndex, int32 KeyframeIndex, const float NewValue);
+
+	/*
+		Set the curve interpolation type (What the curve is doing between keyframes)
+		This should only be used for live-editing features! (changing the curves at runtime is expensive)
+	*/
+	UFUNCTION(BlueprintCallable, Category = "Synth|Components|Audio")
+	bool SetCurveInterpolationType(CurveInterpolationType InterpolationType,  int32 TableIndex);
+
+	/*
+		Set the curve tangent ("Curve depth" between keyframes)
+		This should only be used for live-editing features! (changing the curves at runtime is expensive)
+	*/
+	UFUNCTION(BlueprintCallable, Category = "Synth|Components|Audio")
+	bool SetCurveTangent(int32 TableIndex, float InNewTangent);
+
+	// TODO: Enable this functionality when Curve bug is fixed
+/*
+	UFUNCTION(BlueprintCallable, Category = "Synth|Components|Audio")
+	bool AddTableToEnd();
+
+	UFUNCTION(BlueprintCallable, Category = "Synth|Components|Audio")
+	bool RemoveTableFromEnd();
+*/
+
+// Get the curve interpolation type (What the curve is doing between keyframes)
+	UFUNCTION(BlueprintCallable, Category = "Synth|Components|Audio")
+	float GetCurveTangent(int32 TableIndex);
+
+	// Get an array of floats that represent the key frames in the requested curve
+	UFUNCTION(BlueprintCallable, Category = "Synth|Components|Audio")
+	TArray<float> GetKeyFrameValuesForTable(float TableIndex) const;
+
+	// Wave Table curve was edited
+	UPROPERTY(BlueprintAssignable, Category = "Synth|Preset")
+	FOnTableAltered OnTableAltered;
+
+	// Curve array altered
+	UPROPERTY(BlueprintAssignable, Category = "Synth|Preset")
+	FNumTablesChanged OnNumTablesChanged;
+
+protected:
+	// Helper function that resets a curve to a default saw-tooth
+	// (so a curve is always immediately audible without user effort)
+	void ResetCurve(int32 Index);
+
+	// Initializes the underlying synthesizer.
+	// Called when underlying wave table containers in the Synth FObject need resizing.
+	void InitSynth();
+
+	// Callback for the UAsset this synth component is subscribed to
+	// Compares cached data to current UAsset data and makes updates
+	void ReactToAssetChange(const AssetChangeInfo& ChangeInfo);
+
+#if WITH_EDITOR
+	// Override on Post property change thing here (see if we got a new preset)
+	virtual void PostEditChangeChainProperty(struct FPropertyChangedChainEvent& Eevnt) override;
+#endif // WITH_EDITOR
+
+	/** The settings asset to use for this synth */
+	UPROPERTY(EditAnywhere, Category = "Synth|Components|Audio")
+	UMonoWaveTableSynthPreset* CurrentPreset;
+
+	UMonoWaveTableSynthPreset* CachedPreset;
+
+	// underlying wavetable synth
+	Audio::FMonoWaveTable Synth;
+	Audio::DefaultWaveTableIndexType SampleRate;
+};
\ No newline at end of file
diff --git a/Engine/Plugins/Runtime/Synthesis/Source/Synthesis/Private/SourceEffectPanner.cpp b/Engine/Plugins/Runtime/Synthesis/Source/Synthesis/Private/SourceEffectPanner.cpp
--- a/Engine/Plugins/Runtime/Synthesis/Source/Synthesis/Private/SourceEffectPanner.cpp
+++ b/Engine/Plugins/Runtime/Synthesis/Source/Synthesis/Private/SourceEffectPanner.cpp
@@ -6,7 +6,6 @@
 void FSourceEffectPanner::Init(const FSoundEffectSourceInitData& InitData)
 {
 	bIsActive = true;
-	PanValue = 0.0f;
 	NumChannels = InitData.NumSourceChannels;
 }
 
@@ -14,11 +13,31 @@ void FSourceEffectPanner::OnPresetChanged()
 {
 	GET_EFFECT_SETTINGS(SourceEffectPanner);
 
+	// Normalize the spread value to be between 0.0 and 1.0
+	float SpreadValue = 0.5f * (1.0f + Settings.Spread);
+
+	// Convert to radians between 0.0 and PI/2
+	SpreadValue *= 0.5f * PI;
+
 	// Normalize the panning value to be between 0.0 and 1.0
-	PanValue = 0.5f * (1.0f - Settings.Pan);
+	float PanValue = 0.5f * (1.0f - Settings.Pan);
 
 	// Convert to radians between 0.0 and PI/2
 	PanValue *= 0.5f * PI;
+
+	// Use the "cosine" equal power panning law to compute the spread gain amounts
+	FMath::SinCos(&SpreadGains[0], &SpreadGains[1], SpreadValue);
+
+	// Use the "cosine" equal power panning law to compute a smooth pan based off our parameter
+	FMath::SinCos(&PanGains[0], &PanGains[1], PanValue);
+
+	// Clamp this to be between 0.0 and 1.0 since SinCos is fast and may have values greater than 1.0 or less than 0.0
+	for (int32 i = 0; i < 2; ++i)
+	{
+		SpreadGains[i] = FMath::Clamp(SpreadGains[i], 0.0f, 1.0f);
+		PanGains[i] = FMath::Clamp(PanGains[i], 0.0f, 1.0f);
+	}
+
 }
 
 void FSourceEffectPanner::ProcessAudio(const FSoundEffectSourceInputData& InData, float* OutAudioBufferData)
@@ -29,25 +48,18 @@ void FSourceEffectPanner::ProcessAudio(const FSoundEffectSourceInputData& InData
 	}
 	else
 	{
-		float PanGains[2];
-		// Use the "cosine" equal power panning law to compute a smooth pan based off our parameter
-		FMath::SinCos(&PanGains[0], &PanGains[1], PanValue);
-
-		// Clamp this to be between 0.0 and 1.0 since SinCos is fast and may have values greater than 1.0 or less than 0.0
-		for (int32 i = 0; i < 2; ++i)
-		{
-			PanGains[i] = FMath::Clamp(PanGains[i], 0.0f, 1.0f);
-		}
-
 		for (int32 SampleIndex = 0; SampleIndex < InData.NumSamples; SampleIndex += NumChannels)
 		{
-			// Then simply scale the pan value output with the channel inputs. Note we have to clamp the 
-			// the output of the SinCos since 
-			for (int32 i = 0; i < 2; ++i)
-			{
-				// Simply scale the input sample with the pan value 
-				OutAudioBufferData[SampleIndex + i] = PanGains[i] * InData.InputSourceEffectBufferPtr[SampleIndex + i];
-			}
+			const float LeftChannel = InData.InputSourceEffectBufferPtr[SampleIndex];
+			const float RightChannel = InData.InputSourceEffectBufferPtr[SampleIndex + 1];
+
+			// Left channel spread calculation:
+			// Left channel mix = Sin of Left + Cos of Right
+			// Right channel spread calculation:
+			// Right channel mix = Cos of Left + Sin of Right
+			// Then scale the pan value output with the channel inputs. 
+			OutAudioBufferData[SampleIndex] = PanGains[0] * ((SpreadGains[0] * LeftChannel) + (SpreadGains[1] * RightChannel));
+			OutAudioBufferData[SampleIndex + 1] = PanGains[1] * ((SpreadGains[1] * LeftChannel) + (SpreadGains[0] * RightChannel));
 		}
 	}
 }
diff --git a/Engine/Plugins/Runtime/Synthesis/Source/SynthesisEditor/Private/SynthesisEditorModule.cpp b/Engine/Plugins/Runtime/Synthesis/Source/SynthesisEditor/Private/SynthesisEditorModule.cpp
--- a/Engine/Plugins/Runtime/Synthesis/Source/SynthesisEditor/Private/SynthesisEditorModule.cpp
+++ b/Engine/Plugins/Runtime/Synthesis/Source/SynthesisEditor/Private/SynthesisEditorModule.cpp
@@ -11,6 +11,7 @@
 #include "AssetTypeActions_Base.h"
 #include "AudioEditorModule.h"
 #include "EpicSynth1PresetBank.h"
+#include "MonoWaveTablePresetBank.h"
 
 IMPLEMENT_MODULE(FSynthesisEditorModule, SynthesisEditor)
 
@@ -19,6 +20,7 @@ void FSynthesisEditorModule::StartupModule()
 	IAssetTools& AssetTools = FModuleManager::LoadModuleChecked<FAssetToolsModule>("AssetTools").Get();
 
 	AssetTools.RegisterAssetTypeActions(MakeShareable(new FAssetTypeActions_ModularSynthPresetBank));
+	AssetTools.RegisterAssetTypeActions(MakeShareable(new FAssetTypeActions_MonoWaveTableSynthPreset));
 
 	// Now that we've loaded this module, we need to register our effect preset actions
 	IAudioEditorModule* AudioEditorModule = &FModuleManager::LoadModuleChecked<IAudioEditorModule>("AudioEditor");
diff --git a/Engine/Plugins/Runtime/TimeSynth/Source/TimeSynth/Private/TimeSynthComponent.cpp b/Engine/Plugins/Runtime/TimeSynth/Source/TimeSynth/Private/TimeSynthComponent.cpp
--- a/Engine/Plugins/Runtime/TimeSynth/Source/TimeSynth/Private/TimeSynthComponent.cpp
+++ b/Engine/Plugins/Runtime/TimeSynth/Source/TimeSynth/Private/TimeSynthComponent.cpp
@@ -341,7 +341,7 @@ void UTimeSynthComponent::ShutdownPlayingClips()
 {
 	SoundWaveDecoder.UpdateRenderThread();
 
-	// Loop through all acitve loops and render their audio
+	// Loop through all active loops and render their audio
 	for (int32 i = ActivePlayingClipIndices_AudioRenderThread.Num() - 1; i >= 0; --i)
 	{
 		// Grab the playing clip at the active index
@@ -378,7 +378,7 @@ int32 UTimeSynthComponent::OnGenerateAudio(float* OutAudio, int32 NumSamples)
 	// to the list of active clips. So we only need to loop through active clip indices to render the audio output
 	EventQuantizer.NotifyEvents(NumFrames);
 
-	// Loop through all acitve loops and render their audio
+	// Loop through all active loops and render their audio
 	for (int32 i = ActivePlayingClipIndices_AudioRenderThread.Num() - 1; i >= 0; --i)
 	{
 		// Grab the playing clip at the active index
@@ -515,7 +515,7 @@ int32 UTimeSynthComponent::OnGenerateAudio(float* OutAudio, int32 NumSamples)
 
 void UTimeSynthComponent::SetQuantizationSettings(const FTimeSynthQuantizationSettings& InQuantizationSettings)
 {
-	// Store the quantizaton on the UObject for BP querying
+	// Store the quantization on the UObject for BP querying
 	QuantizationSettings = InQuantizationSettings;
 
 	// Local store what the global quantization is so we can assign it to clips using global quantization
diff --git a/Engine/Source/Editor/AudioEditor/Private/AudioEditorModule.cpp b/Engine/Source/Editor/AudioEditor/Private/AudioEditorModule.cpp
--- a/Engine/Source/Editor/AudioEditor/Private/AudioEditorModule.cpp
+++ b/Engine/Source/Editor/AudioEditor/Private/AudioEditorModule.cpp
@@ -307,6 +307,7 @@ private:
 		SET_AUDIO_ICON(SoundEffectSourcePreset, SourceEffectPreset);
 		SET_AUDIO_ICON(SoundEffectSourcePresetChain, SourceEffectPresetChain_1);
 		SET_AUDIO_ICON(ModularSynthPresetBank, SoundGenericIcon_2);
+		SET_AUDIO_ICON(MonoWaveTableSynthPresetBank, SoundGenericIcon_2);
 		SET_AUDIO_ICON(TimeSynthClip, SoundGenericIcon_2);
 		SET_AUDIO_ICON(TimeSynthVolumeGroup, SoundGenericIcon_1);
 
diff --git a/Engine/Source/Editor/UnrealEd/Private/Commandlets/AudioMixerCommandlet.cpp b/Engine/Source/Editor/UnrealEd/Private/Commandlets/AudioMixerCommandlet.cpp
--- a/Engine/Source/Editor/UnrealEd/Private/Commandlets/AudioMixerCommandlet.cpp
+++ b/Engine/Source/Editor/UnrealEd/Private/Commandlets/AudioMixerCommandlet.cpp
@@ -5,8 +5,8 @@
 #include "Misc/Paths.h"
 #include "Engine/EngineBaseTypes.h"
 #include "Sound/SoundAttenuation.h"
-#include "Audio.h"
 #include "Sound/SoundWave.h"
+#include "Audio.h"
 #include "UObject/UObjectHash.h"
 #include "UObject/Package.h"
 #include "ActiveSound.h"
@@ -161,8 +161,8 @@ static void PlayOneShotSound(FAudioDevice* InAudioDevice, const TArray<USoundWav
 	FActiveSound NewActiveSound;
 	NewActiveSound.SetSound(SoundWave);
 
-	NewActiveSound.VolumeMultiplier = 0.25f;
-	NewActiveSound.PitchMultiplier = FMath::FRandRange(0.1f, 3.0f);
+	NewActiveSound.SetVolume(0.25f);
+	NewActiveSound.SetPitch(FMath::FRandRange(0.1f, 3.0f));
 
 	NewActiveSound.RequestedStartTime = 0.0f;
 
