Fix copyright misspelling

diff --git a/Engine/Plugins/Runtime/Synthesis/Source/Synthesis/Private/SynthComponentMonoWaveTable.cpp b/Engine/Plugins/Runtime/Synthesis/Source/Synthesis/Private/SynthComponentMonoWaveTable.cpp
--- a/Engine/Plugins/Runtime/Synthesis/Source/Synthesis/Private/SynthComponentMonoWaveTable.cpp
+++ b/Engine/Plugins/Runtime/Synthesis/Source/Synthesis/Private/SynthComponentMonoWaveTable.cpp
@@ -1,4 +1,4 @@
-// Copyrigxht 1998-2019 Epic Games, Inc. All Rights Reserved.
+// Copyright 1998-2019 Epic Games, Inc. All Rights Reserved.
 
 
 #include "SynthComponents/SynthComponentMonoWaveTable.h"
