Merging CL 4303647 from Dev-VR - Fixing synth style warnings in Lumin Sample

diff --git a/Engine/Plugins/Runtime/Synthesis/Source/Synthesis/Private/SynthesisModule.cpp b/Engine/Plugins/Runtime/Synthesis/Source/Synthesis/Private/SynthesisModule.cpp
--- a/Engine/Plugins/Runtime/Synthesis/Source/Synthesis/Private/SynthesisModule.cpp
+++ b/Engine/Plugins/Runtime/Synthesis/Source/Synthesis/Private/SynthesisModule.cpp
@@ -3,6 +3,7 @@
 #include "SynthesisModule.h"
 #include "CoreMinimal.h"
 #include "Stats/Stats.h"
+#include "UI/SynthSlateStyle.h"
 #include "Modules/ModuleInterface.h"
 #include "Modules/ModuleManager.h"
 
@@ -12,10 +13,12 @@ IMPLEMENT_MODULE(FSynthesisModule, Synthesis)
 
 void FSynthesisModule::StartupModule()
 {
+	FSynthSlateStyleSet::Initialize();
 }
 
 void FSynthesisModule::ShutdownModule()
 {
+	FSynthSlateStyleSet::Shutdown();
 }
 
 
diff --git a/Engine/Plugins/Runtime/Synthesis/Source/Synthesis/Private/UI/Synth2DSlider.cpp b/Engine/Plugins/Runtime/Synthesis/Source/Synthesis/Private/UI/Synth2DSlider.cpp
--- a/Engine/Plugins/Runtime/Synthesis/Source/Synthesis/Private/UI/Synth2DSlider.cpp
+++ b/Engine/Plugins/Runtime/Synthesis/Source/Synthesis/Private/UI/Synth2DSlider.cpp
@@ -10,6 +10,8 @@
 USynth2DSlider::USynth2DSlider(const FObjectInitializer& ObjectInitializer)
 	: Super(ObjectInitializer)
 {
+	FSynth2DSliderStyle::Initialize();
+
 	SliderHandleColor = FLinearColor::White;
 	StepSize = 0.01f;
 	SSynth2DSlider::FArguments Defaults;
diff --git a/Engine/Plugins/Runtime/Synthesis/Source/Synthesis/Private/UI/Synth2DSliderStyle.cpp b/Engine/Plugins/Runtime/Synthesis/Source/Synthesis/Private/UI/Synth2DSliderStyle.cpp
--- a/Engine/Plugins/Runtime/Synthesis/Source/Synthesis/Private/UI/Synth2DSliderStyle.cpp
+++ b/Engine/Plugins/Runtime/Synthesis/Source/Synthesis/Private/UI/Synth2DSliderStyle.cpp
@@ -4,6 +4,7 @@
 #include "Interfaces/IPluginManager.h"
 #include "SynthesisModule.h"
 #include "CoreMinimal.h"
+#include "UI/SynthSlateStyle.h"
 #include "HAL/FileManager.h"
 
 FSynth2DSliderStyle::FSynth2DSliderStyle()
@@ -14,6 +15,15 @@ FSynth2DSliderStyle::~FSynth2DSliderStyle()
 {
 }
 
+void FSynth2DSliderStyle::Initialize()
+{
+	//first make sure the style set is setup.  Need to make sure because some things happen before StartupModule
+	FSynthSlateStyleSet::Initialize();
+
+	FSynthSlateStyleSet::Get()->Set(FSynth2DSliderStyle::TypeName, FSynth2DSliderStyle::GetDefault());
+}
+
+
 void FSynth2DSliderStyle::GetResources(TArray< const FSlateBrush* >& OutBrushes) const
 {
 }
diff --git a/Engine/Plugins/Runtime/Synthesis/Source/Synthesis/Private/UI/SynthKnob.cpp b/Engine/Plugins/Runtime/Synthesis/Source/Synthesis/Private/UI/SynthKnob.cpp
--- a/Engine/Plugins/Runtime/Synthesis/Source/Synthesis/Private/UI/SynthKnob.cpp
+++ b/Engine/Plugins/Runtime/Synthesis/Source/Synthesis/Private/UI/SynthKnob.cpp
@@ -9,6 +9,8 @@
 USynthKnob::USynthKnob(const FObjectInitializer& ObjectInitializer)
 	: Super(ObjectInitializer)
 {
+	FSynthKnobStyle::Initialize();
+
 	StepSize = 0.01f;
 	MouseSpeed = 1.0f;
 	MouseFineTuneSpeed = 0.2f;
diff --git a/Engine/Plugins/Runtime/Synthesis/Source/Synthesis/Private/UI/SynthKnobStyle.cpp b/Engine/Plugins/Runtime/Synthesis/Source/Synthesis/Private/UI/SynthKnobStyle.cpp
--- a/Engine/Plugins/Runtime/Synthesis/Source/Synthesis/Private/UI/SynthKnobStyle.cpp
+++ b/Engine/Plugins/Runtime/Synthesis/Source/Synthesis/Private/UI/SynthKnobStyle.cpp
@@ -6,6 +6,7 @@
 #include "Internationalization/Regex.h"
 #include "CoreMinimal.h"
 #include "HAL/FileManager.h"
+#include "UI/SynthSlateStyle.h"
 #include "Brushes/SlateDynamicImageBrush.h"
 
 struct FSynthKnobResources
@@ -76,6 +77,16 @@ FSynthKnobStyle::~FSynthKnobStyle()
 {
 }
 
+
+void FSynthKnobStyle::Initialize()
+{
+	//first make sure the style set is setup.  Need to make sure because some things happen before StartupModule
+	FSynthSlateStyleSet::Initialize();
+
+	FSynthSlateStyleSet::Get()->Set(FSynthKnobStyle::TypeName, FSynthKnobStyle::GetDefault());
+}
+
+
 const FSlateBrush* FSynthKnobStyle::GetBaseBrush() const
 {
 	if (KnobSize == ESynthKnobSize::Medium)
diff --git a/Engine/Plugins/Runtime/Synthesis/Source/Synthesis/Private/UI/SynthSlateStyle.cpp b/Engine/Plugins/Runtime/Synthesis/Source/Synthesis/Private/UI/SynthSlateStyle.cpp
--- a/Engine/Plugins/Runtime/Synthesis/Source/Synthesis/Private/UI/SynthSlateStyle.cpp
+++ b/Engine/Plugins/Runtime/Synthesis/Source/Synthesis/Private/UI/SynthSlateStyle.cpp
@@ -5,8 +5,38 @@
 #include "SynthesisModule.h"
 #include "Internationalization/Regex.h"
 #include "HAL/FileManager.h"
+#include "Styling/SlateStyleRegistry.h"
 #include "Brushes/SlateDynamicImageBrush.h"
 
+
+
+TSharedPtr< FSlateStyleSet > FSynthSlateStyleSet::StyleInstance = NULL;
+
+void FSynthSlateStyleSet::Initialize()
+{
+	if (!StyleInstance.IsValid())
+	{
+		StyleInstance = MakeShareable(new FSlateStyleSet("SynthesisStyle"));
+		StyleInstance->SetContentRoot(IPluginManager::Get().FindPlugin("Synthesis")->GetBaseDir() / TEXT("Resources"));
+
+		FSlateStyleRegistry::RegisterSlateStyle(*StyleInstance);
+	}
+}
+
+void FSynthSlateStyleSet::Shutdown()
+{
+	FSlateStyleRegistry::UnRegisterSlateStyle(*StyleInstance);
+	ensure(StyleInstance.IsUnique());
+	StyleInstance.Reset();
+}
+
+TSharedPtr< class FSlateStyleSet > FSynthSlateStyleSet::Get()
+{
+	return StyleInstance; 
+}
+
+
+
 ISynthSlateResources::ISynthSlateResources()
 	: bResourcesLoaded(false)
 {
diff --git a/Engine/Plugins/Runtime/Synthesis/Source/Synthesis/Public/UI/SSynth2DSlider.h b/Engine/Plugins/Runtime/Synthesis/Source/Synthesis/Public/UI/SSynth2DSlider.h
--- a/Engine/Plugins/Runtime/Synthesis/Source/Synthesis/Public/UI/SSynth2DSlider.h
+++ b/Engine/Plugins/Runtime/Synthesis/Source/Synthesis/Public/UI/SSynth2DSlider.h
@@ -12,6 +12,7 @@
 #include "Styling/CoreStyle.h"
 #include "Framework/SlateDelegates.h"
 #include "Widgets/SLeafWidget.h"
+#include "Styling/SlateStyle.h"
 #include "UI/Synth2DSliderStyle.h"
 
 class FPaintArgs;
@@ -28,7 +29,7 @@ public:
 	SLATE_BEGIN_ARGS(SSynth2DSlider)
 		: _IndentHandle(true)
 		, _Locked(false)
-		, _Style(&FCoreStyle::Get().GetWidgetStyle<FSynth2DSliderStyle>("Synth2DSliderStyle"))
+		, _Style(&FSynthSlateStyleSet::Get()->GetWidgetStyle<FSynth2DSliderStyle>("Synth2DSliderStyle"))
 		, _StepSize(0.01f)
 		, _ValueX(1.f)
 		, _ValueY(1.f)
diff --git a/Engine/Plugins/Runtime/Synthesis/Source/Synthesis/Public/UI/SSynthKnob.h b/Engine/Plugins/Runtime/Synthesis/Source/Synthesis/Public/UI/SSynthKnob.h
--- a/Engine/Plugins/Runtime/Synthesis/Source/Synthesis/Public/UI/SSynthKnob.h
+++ b/Engine/Plugins/Runtime/Synthesis/Source/Synthesis/Public/UI/SSynthKnob.h
@@ -25,7 +25,7 @@ public:
 	
 	SLATE_BEGIN_ARGS(SSynthKnob)
 		: _Locked(false)
-		, _Style(&FCoreStyle::Get().GetWidgetStyle<FSynthKnobStyle>("SynthKnobStyle"))
+		, _Style(&FSynthSlateStyleSet::Get()->GetWidgetStyle<FSynthKnobStyle>("SynthKnobStyle"))
 		, _StepSize(0.01f)
 		, _MouseSpeed(1.0f)
 		, _MouseFineTuneSpeed(0.2f)
diff --git a/Engine/Plugins/Runtime/Synthesis/Source/Synthesis/Public/UI/Synth2DSliderStyle.h b/Engine/Plugins/Runtime/Synthesis/Source/Synthesis/Public/UI/Synth2DSliderStyle.h
--- a/Engine/Plugins/Runtime/Synthesis/Source/Synthesis/Public/UI/Synth2DSliderStyle.h
+++ b/Engine/Plugins/Runtime/Synthesis/Source/Synthesis/Public/UI/Synth2DSliderStyle.h
@@ -26,6 +26,8 @@ struct SYNTHESIS_API FSynth2DSliderStyle : public FSlateWidgetStyle
 
 	virtual ~FSynth2DSliderStyle();
 
+	static void Initialize();
+
 	// Image to use for the 2D handle
 	UPROPERTY(EditAnywhere, BlueprintReadWrite, Category = Appearance)
 	FSlateBrush NormalThumbImage;
diff --git a/Engine/Plugins/Runtime/Synthesis/Source/Synthesis/Public/UI/SynthKnobStyle.h b/Engine/Plugins/Runtime/Synthesis/Source/Synthesis/Public/UI/SynthKnobStyle.h
--- a/Engine/Plugins/Runtime/Synthesis/Source/Synthesis/Public/UI/SynthKnobStyle.h
+++ b/Engine/Plugins/Runtime/Synthesis/Source/Synthesis/Public/UI/SynthKnobStyle.h
@@ -4,6 +4,8 @@
 
 #include "CoreMinimal.h"
 
+#include "UI/SynthSlateStyle.h"
+
 #include "UObject/ObjectMacros.h"
 #include "Styling/SlateTypes.h"
 #include "Widgets/SWidget.h"
@@ -32,6 +34,8 @@ struct SYNTHESIS_API FSynthKnobStyle : public FSlateWidgetStyle
 
 	virtual ~FSynthKnobStyle();
 
+	static void FSynthKnobStyle::Initialize();
+
 	static const FName TypeName;
 	virtual const FName GetTypeName() const override { return TypeName; };
 
diff --git a/Engine/Plugins/Runtime/Synthesis/Source/Synthesis/Public/UI/SynthSlateStyle.h b/Engine/Plugins/Runtime/Synthesis/Source/Synthesis/Public/UI/SynthSlateStyle.h
--- a/Engine/Plugins/Runtime/Synthesis/Source/Synthesis/Public/UI/SynthSlateStyle.h
+++ b/Engine/Plugins/Runtime/Synthesis/Source/Synthesis/Public/UI/SynthSlateStyle.h
@@ -3,12 +3,12 @@
 #pragma once
 
 #include "CoreMinimal.h"
-#include "UI/SSynthKnob.h"
 #include "UObject/ObjectMacros.h"
 #include "Styling/SlateTypes.h"
 #include "Widgets/SWidget.h"
 #include "Components/Widget.h"
 #include "Styling/SlateWidgetStyle.h"
+#include "Styling/SlateStyle.h"
 #include "Brushes/SlateImageBrush.h"
 #include "SynthSlateStyle.generated.h"
 
@@ -47,6 +47,18 @@ protected:
 	bool bResourcesLoaded;
 };
 
+
+struct FSynthSlateStyleSet
+{
+	static void Initialize();
+	static void Shutdown();
+
+	static TSharedPtr< class FSlateStyleSet > Get();
+
+private:
+	static TSharedPtr< class FSlateStyleSet > StyleInstance;
+};
+
 /**
  * Represents the appearance of synth UI elements in slate.
  */
