Merging CL 4303846 from Dev-VR - Fix for fully qualified name in header

diff --git a/Engine/Plugins/Runtime/Synthesis/Source/Synthesis/Public/UI/SynthKnobStyle.h b/Engine/Plugins/Runtime/Synthesis/Source/Synthesis/Public/UI/SynthKnobStyle.h
--- a/Engine/Plugins/Runtime/Synthesis/Source/Synthesis/Public/UI/SynthKnobStyle.h
+++ b/Engine/Plugins/Runtime/Synthesis/Source/Synthesis/Public/UI/SynthKnobStyle.h
@@ -34,7 +34,7 @@ struct SYNTHESIS_API FSynthKnobStyle : public FSlateWidgetStyle
 
 	virtual ~FSynthKnobStyle();
 
-	static void FSynthKnobStyle::Initialize();
+	static void Initialize();
 
 	static const FName TypeName;
 	virtual const FName GetTypeName() const override { return TypeName; };
