[NO NON-PII COMMIT MESSAGE]

diff --git a/Engine/Plugins/Runtime/Synthesis/Source/Synthesis/Classes/SynthComponents/EpicSynth1Component.h b/Engine/Plugins/Runtime/Synthesis/Source/Synthesis/Classes/SynthComponents/EpicSynth1Component.h
--- a/Engine/Plugins/Runtime/Synthesis/Source/Synthesis/Classes/SynthComponents/EpicSynth1Component.h
+++ b/Engine/Plugins/Runtime/Synthesis/Source/Synthesis/Classes/SynthComponents/EpicSynth1Component.h
@@ -14,6 +14,7 @@
 #include "EpicSynth1.h"
 #include "EpicSynth1Types.h"
 #include "DSP/Osc.h"
+#include "Kismet/BlueprintFunctionLibrary.h"
 #include "EpicSynth1Component.generated.h"
 
 USTRUCT(BlueprintType)
@@ -361,6 +362,16 @@ public:
  	TArray<FModularSynthPresetBankEntry> Presets;
 };
 
+UCLASS()
+class SYNTHESIS_API UModularSynthLibrary : public UBlueprintFunctionLibrary
+{
+	GENERATED_BODY()
+
+	// Adds the modular synth preset to the bank asset in the content browser. Only call during editor.
+	UFUNCTION(BlueprintCallable, Category = "Synthesis")
+	static void AddModularSynthPresetToBankAsset(UModularSynthPresetBank* InBank, const FModularSynthPreset& Preset, const FString& PresetName);
+};
+
 
 /**
 * UModularSynthComponent
diff --git a/Engine/Plugins/Runtime/Synthesis/Source/Synthesis/Private/EpicSynth1.cpp b/Engine/Plugins/Runtime/Synthesis/Source/Synthesis/Private/EpicSynth1.cpp
--- a/Engine/Plugins/Runtime/Synthesis/Source/Synthesis/Private/EpicSynth1.cpp
+++ b/Engine/Plugins/Runtime/Synthesis/Source/Synthesis/Private/EpicSynth1.cpp
@@ -395,7 +395,7 @@ namespace Audio
 
 	void FEpicSynth1Voice::ClearPatches()
 	{
-
+		// Clear dynamic patches
 		FModulationMatrix* ModMatrix = &ParentSynth->ModMatrix;
 		for (auto PatchEntry : DynamicPatches)
 		{
@@ -452,6 +452,18 @@ namespace Audio
 		bIsFinished = true;
 		bIsActive = false;
 		VoiceGeneration = INDEX_NONE;
+
+		for (int32 i = 0; i < NumOscillators; ++i)
+		{
+			Oscil[i].Reset();
+			OscilPan[i].Reset();
+		}
+
+		OnePoleFilter.Reset();
+		StateVarFilter.Reset();
+		LadderFilter.Reset();
+
+		Amp.Reset();
 	}
 
 	void FEpicSynth1Voice::NoteOn(const uint32 InMidiNote, const float InVelocity, const float InDurationSec)
@@ -848,18 +860,32 @@ namespace Audio
 			else if (FreeVoices.Num() > 0)
 			{
 				uint32 OldestId = GetOldestPlayingId();
-				Voices[OldestId]->Shutdown();
+				if (OldestId < (uint32)Voices.Num())
+				{
+					Voices[OldestId]->Shutdown();
+				}
 				VoiceId = FreeVoices.Pop();
 				Voice = Voices[VoiceId];
 			}
 			else
 			{
 				uint32 OldestId = GetOldestPlayingId();
-				Voices[OldestId]->Kill();
-				--NumActiveVoices;
+				if (OldestId < (uint32)Voices.Num())
+				{
+					Voices[OldestId]->Kill();
+					--NumActiveVoices;
+
+					VoiceId = OldestId;
+					Voice = Voices[OldestId];
+				}
+			}
 
-				VoiceId = OldestId;
-				Voice = Voices[OldestId];
+			// Fallback to grabbing the 0th voice
+			if (VoiceId == INDEX_NONE && Voices.Num() > 0)
+			{
+				Voices[0]->Kill();
+				VoiceId = 0;
+				Voice = Voices[0];
 			}
 
 			++NumActiveVoices;
@@ -1422,8 +1448,11 @@ namespace Audio
 
 	void FEpicSynth1::ClearPatches()
 	{
+		ModMatrix.ResetPatchSourceState();
+
 		for (int32 VoiceId = 0; VoiceId < Voices.Num(); ++VoiceId)
 		{
+			Voices[VoiceId]->Reset();
 			Voices[VoiceId]->ClearPatches();
 		}
 	}
diff --git a/Engine/Plugins/Runtime/Synthesis/Source/Synthesis/Private/EpicSynth1Component.cpp b/Engine/Plugins/Runtime/Synthesis/Source/Synthesis/Private/EpicSynth1Component.cpp
--- a/Engine/Plugins/Runtime/Synthesis/Source/Synthesis/Private/EpicSynth1Component.cpp
+++ b/Engine/Plugins/Runtime/Synthesis/Source/Synthesis/Private/EpicSynth1Component.cpp
@@ -613,3 +613,16 @@ int32 UModularSynthComponent::OnGenerateAudio(float* OutAudio, int32 NumSamples)
 	}
 	return NumSamples;
 }
+
+
+void UModularSynthLibrary::AddModularSynthPresetToBankAsset(UModularSynthPresetBank* InBank, const FModularSynthPreset& Preset, const FString& PresetName)
+{
+	if (GIsEditor)
+	{
+		FModularSynthPresetBankEntry NewBankEntry;
+		NewBankEntry.PresetName = PresetName;
+		NewBankEntry.Preset = Preset;
+
+		InBank->Presets.Add(NewBankEntry);
+	}
+}
diff --git a/Engine/Source/Runtime/AudioMixer/Private/DSP/Amp.cpp b/Engine/Source/Runtime/AudioMixer/Private/DSP/Amp.cpp
--- a/Engine/Source/Runtime/AudioMixer/Private/DSP/Amp.cpp
+++ b/Engine/Source/Runtime/AudioMixer/Private/DSP/Amp.cpp
@@ -194,8 +194,10 @@ namespace Audio
 
 	void FAmp::Reset()
 	{
-		GainEnv = 0.0f;
+		GainEnv = 1.0f;
 		GainMod = 1.0f;
+		Pan = 0.0f;
+		PanMod = 0.0f;
 	}
 
 }
diff --git a/Engine/Source/Runtime/AudioMixer/Private/DSP/Filter.cpp b/Engine/Source/Runtime/AudioMixer/Private/DSP/Filter.cpp
--- a/Engine/Source/Runtime/AudioMixer/Private/DSP/Filter.cpp
+++ b/Engine/Source/Runtime/AudioMixer/Private/DSP/Filter.cpp
@@ -288,6 +288,15 @@ namespace Audio
 		}
 	}
 
+	void IFilter::Reset()
+	{
+		ModQ = 0.0f;
+		ExternalModQ = 0.0f;
+		ExternalModFrequency = 0.0f;
+		ModFrequency = 0.0f;
+		bChanged = true;
+	}
+
 	void IFilter::SetFrequency(const float InCutoffFrequency)
 	{
 		if (BaseFrequency != InCutoffFrequency)
@@ -377,6 +386,8 @@ namespace Audio
 
 	void FOnePoleFilter::Reset()
 	{
+		IFilter::Reset();
+
 		FMemory::Memzero(Z1, sizeof(float)*NumChannels);
 	}
 
@@ -476,6 +487,8 @@ namespace Audio
 
 	void FStateVariableFilter::Reset()
 	{
+		IFilter::Reset();
+
 		FMemory::Memzero(FilterState.GetData(), sizeof(FFilterState) * NumChannels);
 	}
 
@@ -561,6 +574,8 @@ namespace Audio
 
 	void FLadderFilter::Reset()
 	{
+		IFilter::Reset();
+
 		for (int32 i = 0; i < 4; ++i)
 		{
 			OnePoleFilters[i].Reset();
diff --git a/Engine/Source/Runtime/AudioMixer/Private/DSP/ModulationMatrix.cpp b/Engine/Source/Runtime/AudioMixer/Private/DSP/ModulationMatrix.cpp
--- a/Engine/Source/Runtime/AudioMixer/Private/DSP/ModulationMatrix.cpp
+++ b/Engine/Source/Runtime/AudioMixer/Private/DSP/ModulationMatrix.cpp
@@ -93,6 +93,17 @@ namespace Audio
 		return true;
 	}
 
+	void FModulationMatrix::ResetPatchSourceState()
+	{
+		for (TArray<float>& SourceVoiceStates : Sources)
+		{
+			for (float& Value : SourceVoiceStates)
+			{
+				Value = 0.0f;
+			}
+		}
+	}
+
 	void FModulationMatrix::ClearPatches(const int32 VoiceId)
 	{
 		check(VoiceId < NumVoices);
diff --git a/Engine/Source/Runtime/AudioMixer/Private/DSP/Osc.cpp b/Engine/Source/Runtime/AudioMixer/Private/DSP/Osc.cpp
--- a/Engine/Source/Runtime/AudioMixer/Private/DSP/Osc.cpp
+++ b/Engine/Source/Runtime/AudioMixer/Private/DSP/Osc.cpp
@@ -180,8 +180,12 @@ namespace Audio
 	{
 		Phase = 0.0f;
 		ExternalGainMod = 1.0f;
+		PulseWidthMod = 0.0f;
 		FreqData.PitchBend = 0.0f;
+		FreqData.Mod = 0.0f;
+		FreqData.ExternalMod = 0.0f;
 		FreqData.Detune = 0.0f;
+		bChanged = true;
 	}
 
 	FOsc::FOsc()
diff --git a/Engine/Source/Runtime/AudioMixer/Public/DSP/Filter.h b/Engine/Source/Runtime/AudioMixer/Public/DSP/Filter.h
--- a/Engine/Source/Runtime/AudioMixer/Public/DSP/Filter.h
+++ b/Engine/Source/Runtime/AudioMixer/Public/DSP/Filter.h
@@ -145,13 +145,16 @@ namespace Audio
 		virtual void SetFilterType(const EFilter::Type InFilterType);
 
 		// Reset the filter
-		virtual void Reset() {}
+		virtual void Reset();
 
 		// Updates the filter
 		virtual void Update();
 
 		// Processes a single frame of audio. Number of channels MUST be what was used during filter initialization.
-		virtual void ProcessAudioFrame(const float* InFrame, float* OutFrame) {}
+		virtual void ProcessAudioFrame(const float* InFrame, float* OutFrame)
+		{
+			ProcessAudio(InFrame, NumChannels, OutFrame);
+		}
 
 		// Process an audio buffer.
 		virtual void ProcessAudio(const float* InBuffer, const int32 InNumSamples, float* OutBuffer) = 0;
diff --git a/Engine/Source/Runtime/AudioMixer/Public/DSP/ModulationMatrix.h b/Engine/Source/Runtime/AudioMixer/Public/DSP/ModulationMatrix.h
--- a/Engine/Source/Runtime/AudioMixer/Public/DSP/ModulationMatrix.h
+++ b/Engine/Source/Runtime/AudioMixer/Public/DSP/ModulationMatrix.h
@@ -122,6 +122,9 @@ namespace Audio
 		// Removes the given patch connection between a source and one or more destinations
 		bool RemovePatch(const int32 VoiceId, FPatch* Patch);
 
+		// Reset all patch state
+		void ResetPatchSourceState();
+
 		// Clear all patch connections
 		void ClearPatches(const int32 VoiceId);
 
