[NO NON-PII COMMIT MESSAGE]

diff --git a/Engine/Plugins/Runtime/Synthesis/Source/Synthesis/Private/EpicSynth1Component.cpp b/Engine/Plugins/Runtime/Synthesis/Source/Synthesis/Private/EpicSynth1Component.cpp
--- a/Engine/Plugins/Runtime/Synthesis/Source/Synthesis/Private/EpicSynth1Component.cpp
+++ b/Engine/Plugins/Runtime/Synthesis/Source/Synthesis/Private/EpicSynth1Component.cpp
@@ -624,5 +624,6 @@ void UModularSynthLibrary::AddModularSynthPresetToBankAsset(UModularSynthPresetB
 		NewBankEntry.Preset = Preset;
 
 		InBank->Presets.Add(NewBankEntry);
+		InBank->MarkPackageDirty();
 	}
 }
