Fix initializer warnings #rb none #jira UE-70052

diff --git a/Engine/Plugins/Runtime/Synthesis/Source/Synthesis/Classes/SourceEffects/SourceEffectPhaser.h b/Engine/Plugins/Runtime/Synthesis/Source/Synthesis/Classes/SourceEffects/SourceEffectPhaser.h
--- a/Engine/Plugins/Runtime/Synthesis/Source/Synthesis/Classes/SourceEffects/SourceEffectPhaser.h
+++ b/Engine/Plugins/Runtime/Synthesis/Source/Synthesis/Classes/SourceEffects/SourceEffectPhaser.h
@@ -49,6 +49,7 @@ struct SYNTHESIS_API FSourceEffectPhaserSettings
 		, Frequency(2.0f)
 		, Feedback(0.3f)
 		, LFOType(EPhaserLFOType::Sine)
+		, UseQuadraturePhase(false)
 	{}
 };
 
diff --git a/Engine/Plugins/Runtime/Synthesis/Source/Synthesis/Classes/SourceEffects/SourceEffectStereoDelay.h b/Engine/Plugins/Runtime/Synthesis/Source/Synthesis/Classes/SourceEffects/SourceEffectStereoDelay.h
--- a/Engine/Plugins/Runtime/Synthesis/Source/Synthesis/Classes/SourceEffects/SourceEffectStereoDelay.h
+++ b/Engine/Plugins/Runtime/Synthesis/Source/Synthesis/Classes/SourceEffects/SourceEffectStereoDelay.h
@@ -52,6 +52,7 @@ struct SYNTHESIS_API FSourceEffectStereoDelaySettings
 	FSourceEffectStereoDelaySettings()
 		: DelayMode(EStereoDelaySourceEffect::PingPong)
 		, DelayTimeMsec(500.0f)
+		, Feedback(0.1f)
 		, DelayRatio(0.2f)
 		, WetLevel(0.4f)
 	{}
diff --git a/Engine/Plugins/Runtime/Synthesis/Source/Synthesis/Private/SubmixEffectTapDelay.cpp b/Engine/Plugins/Runtime/Synthesis/Source/Synthesis/Private/SubmixEffectTapDelay.cpp
--- a/Engine/Plugins/Runtime/Synthesis/Source/Synthesis/Private/SubmixEffectTapDelay.cpp
+++ b/Engine/Plugins/Runtime/Synthesis/Source/Synthesis/Private/SubmixEffectTapDelay.cpp
@@ -73,8 +73,10 @@ FTapDelayInfo::FTapDelayInfo()
 	, Gain(-3.0f)
 	, OutputChannel(0)
 	, PanInDegrees(0.0f)
-	, TapId(TapIdCount++)
-{}
+	, TapId(0)
+{
+	TapId = TapIdCount++;
+}
 
 FSubmixEffectTapDelay::FSubmixEffectTapDelay()
 	: SampleRate(0.0f)
diff --git a/Engine/Plugins/Runtime/Synthesis/Source/Synthesis/Private/UI/Synth2DSliderStyle.cpp b/Engine/Plugins/Runtime/Synthesis/Source/Synthesis/Private/UI/Synth2DSliderStyle.cpp
--- a/Engine/Plugins/Runtime/Synthesis/Source/Synthesis/Private/UI/Synth2DSliderStyle.cpp
+++ b/Engine/Plugins/Runtime/Synthesis/Source/Synthesis/Private/UI/Synth2DSliderStyle.cpp
@@ -8,6 +8,7 @@
 #include "HAL/FileManager.h"
 
 FSynth2DSliderStyle::FSynth2DSliderStyle()
+	: BarThickness(2.0f)
 {
 }
 
