Fix for SynthKnob using bad rotation ceterpoint.

diff --git a/Engine/Plugins/Runtime/Synthesis/Source/Synthesis/Private/UI/SSynthKnob.cpp b/Engine/Plugins/Runtime/Synthesis/Source/Synthesis/Private/UI/SSynthKnob.cpp
--- a/Engine/Plugins/Runtime/Synthesis/Source/Synthesis/Private/UI/SSynthKnob.cpp
+++ b/Engine/Plugins/Runtime/Synthesis/Source/Synthesis/Private/UI/SSynthKnob.cpp
@@ -51,6 +51,7 @@ int32 SSynthKnob::OnPaint(const FPaintArgs& Args, const FGeometry& AllottedGeome
 
 	const FSlateBrush* BaseImageBrush = Style->GetBaseBrush();
 	const FSlateBrush* OverlayImageBrush = Style->GetOverlayBrush();
+	const FVector2D ImageSize = BaseImageBrush->GetImageSize();
 
 	if (BaseImageBrush != nullptr)
 	{
@@ -63,12 +64,13 @@ int32 SSynthKnob::OnPaint(const FPaintArgs& Args, const FGeometry& AllottedGeome
 		const float MinValueAngle = Style->MinValueAngle;
 		const float MaxValueAngle = Style->MaxValueAngle;
 		const float CurrentValue = GetValue();
+		const FVector2D ImageCenter = ImageSize * 0.5f;
 
 		float NormalizedRotationAngle = CurrentValue * (MaxValueAngle - MinValueAngle) + MinValueAngle;
 		float RotationAngle = 2.0f * PI * NormalizedRotationAngle;
-
+		
 		const FLinearColor FinalColorAndOpacity(InWidgetStyle.GetColorAndOpacityTint() * OverlayImageBrush->GetTint(InWidgetStyle));
-		FSlateDrawElement::MakeRotatedBox(OutDrawElements, LayerId++, AllottedGeometry.ToPaintGeometry(), OverlayImageBrush, DrawEffects, RotationAngle, MyCullingRect.GetCenter(), FSlateDrawElement::RelativeToWorld, FinalColorAndOpacity);
+		FSlateDrawElement::MakeRotatedBox(OutDrawElements, LayerId++, AllottedGeometry.ToPaintGeometry(), OverlayImageBrush, DrawEffects, RotationAngle, ImageCenter, FSlateDrawElement::RelativeToElement, FinalColorAndOpacity);
 	}
 
 	return LayerId;
