Fix for static analysis warning.
Fix for bug when resizing widget.

diff --git a/Engine/Plugins/Runtime/Synthesis/Source/Synthesis/Private/UI/SSynthKnob.cpp b/Engine/Plugins/Runtime/Synthesis/Source/Synthesis/Private/UI/SSynthKnob.cpp
--- a/Engine/Plugins/Runtime/Synthesis/Source/Synthesis/Private/UI/SSynthKnob.cpp
+++ b/Engine/Plugins/Runtime/Synthesis/Source/Synthesis/Private/UI/SSynthKnob.cpp
@@ -51,7 +51,6 @@ int32 SSynthKnob::OnPaint(const FPaintArgs& Args, const FGeometry& AllottedGeome
 
 	const FSlateBrush* BaseImageBrush = Style->GetBaseBrush();
 	const FSlateBrush* OverlayImageBrush = Style->GetOverlayBrush();
-	const FVector2D ImageSize = BaseImageBrush->GetImageSize();
 
 	if (BaseImageBrush != nullptr)
 	{
@@ -64,7 +63,7 @@ int32 SSynthKnob::OnPaint(const FPaintArgs& Args, const FGeometry& AllottedGeome
 		const float MinValueAngle = Style->MinValueAngle;
 		const float MaxValueAngle = Style->MaxValueAngle;
 		const float CurrentValue = GetValue();
-		const FVector2D ImageCenter = ImageSize * 0.5f;
+		const FVector2D ImageCenter = AllottedGeometry.GetLocalSize() * 0.5f;
 
 		float NormalizedRotationAngle = CurrentValue * (MaxValueAngle - MinValueAngle) + MinValueAngle;
 		float RotationAngle = 2.0f * PI * NormalizedRotationAngle;
