Fixed wave table normalization bug in Wave Table Synth (adding removal of DC offset)
Fixed unwanted edge-case behavior when clearing an already empty table array

diff --git a/Engine/Plugins/Runtime/Synthesis/Source/Synthesis/Classes/SynthComponents/SynthComponentMonoWaveTable.h b/Engine/Plugins/Runtime/Synthesis/Source/Synthesis/Classes/SynthComponents/SynthComponentMonoWaveTable.h
--- a/Engine/Plugins/Runtime/Synthesis/Source/Synthesis/Classes/SynthComponents/SynthComponentMonoWaveTable.h
+++ b/Engine/Plugins/Runtime/Synthesis/Source/Synthesis/Classes/SynthComponents/SynthComponentMonoWaveTable.h
@@ -21,13 +21,6 @@
 #include "SynthesisModule.h"
 #include "SynthComponentMonoWaveTable.generated.h"
 
-// ========================================================================
-// USynthComponentMonoWaveTable
-// Synth component class which implements USynthComponent
-// This is a simple hello-world type example which generates a sine-wave
-// tone using a DSP oscillator class and implements a single function to set
-// the frequency. 
-// ========================================================================
 
 class USynthComponentMonoWaveTable;
 
@@ -128,7 +121,6 @@ protected:
 
 	// cached asset data (before last edit)
 	int8 bWasLockedToGrid : 1;
-	int8 bWasSingleTable : 1;
 	int32 CachedTableResolution;
 	TArray<FRuntimeFloatCurve> CachedWaveTable;
 	uint8 bCachedNormalizationSetting : 1;
diff --git a/Engine/Plugins/Runtime/Synthesis/Source/Synthesis/Private/SynthComponentMonoWaveTable.cpp b/Engine/Plugins/Runtime/Synthesis/Source/Synthesis/Private/SynthComponentMonoWaveTable.cpp
--- a/Engine/Plugins/Runtime/Synthesis/Source/Synthesis/Private/SynthComponentMonoWaveTable.cpp
+++ b/Engine/Plugins/Runtime/Synthesis/Source/Synthesis/Private/SynthComponentMonoWaveTable.cpp
@@ -23,7 +23,6 @@ UMonoWaveTableSynthPreset::UMonoWaveTableSynthPreset() : PresetName(TEXT("Defaul
 	DefaultCurvePointer->AddKey(1.0f, 0.0f);
 
 	bWasLockedToGrid = bLockKeyframesToGridBool;
-	bWasSingleTable = true;
 
 	CacheAssetData();
 }
@@ -48,16 +47,13 @@ void UMonoWaveTableSynthPreset::PostEditChangeChainProperty(FPropertyChangedChai
 void UMonoWaveTableSynthPreset::EditChangeInternal()
 {
 	AssetChangeInfo ChangeInfo;
+	
 
-	// Adding the first table
-	if (!WaveTable.Num() && !bWasSingleTable)
+	if (!WaveTable.Num())
 	{
-		WaveTable.Add(DefaultCurve);
 		ChangeInfo.bNeedsFullRebuild = true;
 	}
 
-	bWasSingleTable = (WaveTable.Num() == 1);
-
 	// resize tangent data if array length increased
 	if (CurveBiDirTangents.Num() < WaveTable.Num())
 	{
@@ -330,7 +326,6 @@ void UMonoWaveTableSynthPreset::CacheAssetData()
 		bCachedNormalizationSetting = bNormalizeWaveTables;
 		CachedWaveTable = WaveTable;
 		CachedTableResolution = WaveTableResolution;
-		bWasSingleTable = (WaveTable.Num() == 1);
 }
 
 bool UMonoWaveTableSynthPreset::IsCachedTableEntryStillValid(Audio::DefaultWaveTableIndexType Index)
@@ -406,32 +401,45 @@ void USynthComponentMonoWaveTable::RefreshWaveTable(int32 Index)
 	double CurrTime = 0.0;
 
 	// for (optional) normalization
-	float MaxAbsValue = 0.0f;
+	float MaxValue = 0.0f;
+	float MinValue = 0.0f;
+
+	for (Audio::DefaultWaveTableIndexType i = 0; i < TableRes; ++i)
+	{
+		DestinationTable[i] = RichCurveToProcess.Eval(CurrTime);
+		const float CurrSample = DestinationTable[i];
+
+		MaxValue = CurrSample > MaxValue ? CurrSample : MaxValue;
+		MinValue = CurrSample < MinValue ? CurrSample : MinValue;
+		CurrTime += SamplePeriod;
+	}
+
+	const float DCShift = -0.5f * (FMath::Abs(MaxValue) - FMath::Abs(MinValue));
+	MaxValue += DCShift;
+	MinValue += DCShift;
+
+	// Remove DC offset
+	if (!FMath::IsNearlyZero(DCShift))
+	{
+		Audio::AddConstantToBufferInplace(DestinationTable.GetData(), TableRes, DCShift);
+	}
 
-	// Normalize audio when we sample
 	if (CachedPreset->bNormalizeWaveTables)
 	{
-		for (Audio::DefaultWaveTableIndexType i = 0; i < TableRes; ++i)
-		{
-			DestinationTable[i] = RichCurveToProcess.Eval(CurrTime);
-			MaxAbsValue = FMath::Abs(DestinationTable[i]) > MaxAbsValue ? DestinationTable[i] : MaxAbsValue;
-			CurrTime += SamplePeriod;
-		}
+		const float MaxAbsValue = FMath::Max(FMath::Abs(MaxValue), FMath::Abs(MinValue));
+		const float NormalizationScalar = 1.0f / MaxAbsValue;
 
-		float NormalizationScalar = 1.0f / MaxAbsValue;
 		Audio::MultiplyBufferByConstantInPlace(DestinationTable.GetData(), TableRes, NormalizationScalar);
 	}
-	// clip audio while we sample
-	else
+	else // clip the values since we aren't normalizing
 	{
-		for (Audio::DefaultWaveTableIndexType i = 0; i < TableRes; ++i)
+		for (auto& Sample : DestinationTable)
 		{
-			const float CurrSample = RichCurveToProcess.Eval(CurrTime);
-			DestinationTable[i] = FMath::Clamp(CurrSample, -1.0f, 1.0f);
-			CurrTime += SamplePeriod;
+			Sample = FMath::Clamp(Sample, -1.0f, 1.0f);
 		}
 	}
 
+
 	// push sampled data to synth
 	// (Note: only performs simple memcopy in pump)
 	SynthCommand([this, Index, DestinationTable{ MoveTemp(DestinationTable) }]()
diff --git a/Engine/Source/Runtime/AudioMixer/Private/DSP/BufferVectorOperations.cpp b/Engine/Source/Runtime/AudioMixer/Private/DSP/BufferVectorOperations.cpp
--- a/Engine/Source/Runtime/AudioMixer/Private/DSP/BufferVectorOperations.cpp
+++ b/Engine/Source/Runtime/AudioMixer/Private/DSP/BufferVectorOperations.cpp
@@ -108,6 +108,24 @@ namespace Audio
 		}
 	}
 
+	// Adds a constant to a buffer (useful for DC offset removal)
+	void AddConstantToBufferInplace(AlignedFloatBuffer& InBuffer, float InConstant)
+	{
+		AddConstantToBufferInplace(InBuffer.GetData(), InBuffer.Num(), InConstant);
+	}
+
+	void AddConstantToBufferInplace(float* RESTRICT InBuffer, int32 NumSamples, float InConstant)
+	{
+		const VectorRegister Constant = VectorLoadFloat1(&InConstant);
+
+		for (int32 i = 0; i < NumSamples; i += 4)
+		{
+			VectorRegister Output = VectorLoadAligned(&InBuffer[i]);
+			Output = VectorAdd(Output, Constant);
+			VectorStoreAligned(Output, &InBuffer[i]);
+		}
+	}
+
 	/* Performs an element-wise weighted sum OutputBuffer = (InBuffer1 x InGain1) + (InBuffer2 x InGain2) */
 	void BufferWeightedSumFast(const AlignedFloatBuffer& InBuffer1, float InGain1, const AlignedFloatBuffer& InBuffer2, float InGain2, AlignedFloatBuffer& OutBuffer)
 	{
diff --git a/Engine/Source/Runtime/AudioMixer/Public/DSP/BufferVectorOperations.h b/Engine/Source/Runtime/AudioMixer/Public/DSP/BufferVectorOperations.h
--- a/Engine/Source/Runtime/AudioMixer/Public/DSP/BufferVectorOperations.h
+++ b/Engine/Source/Runtime/AudioMixer/Public/DSP/BufferVectorOperations.h
@@ -38,6 +38,10 @@ namespace Audio
 	void AUDIOMIXER_API MultiplyBufferByConstantInPlace(AlignedFloatBuffer& InBuffer, float InGain);
 	void AUDIOMIXER_API MultiplyBufferByConstantInPlace(float* RESTRICT InBuffer, int32 NumSamples, float InGain);
 
+	/** Adds a constant to a buffer (useful for DC offset removal) */
+	void AUDIOMIXER_API AddConstantToBufferInplace(AlignedFloatBuffer& InBuffer, float Constant);
+	void AUDIOMIXER_API AddConstantToBufferInplace(float* RESTRICT InBuffer, int32 NumSamples, float Constant);
+
 	/* Performs an element-wise weighted sum OutputBuffer = (InBuffer1 x InGain1) + (InBuffer2 x InGain2) */
 	void AUDIOMIXER_API BufferWeightedSumFast(const AlignedFloatBuffer& InBuffer1, float InGain1, const AlignedFloatBuffer& InBuffer2, float InGain2, AlignedFloatBuffer& OutBuffer);
 
