Setting icon and color for WaveTableSynth preset asset

diff --git a/Engine/Plugins/Runtime/Synthesis/Source/SynthesisEditor/Classes/MonoWaveTablePresetBank.h b/Engine/Plugins/Runtime/Synthesis/Source/SynthesisEditor/Classes/MonoWaveTablePresetBank.h
--- a/Engine/Plugins/Runtime/Synthesis/Source/SynthesisEditor/Classes/MonoWaveTablePresetBank.h
+++ b/Engine/Plugins/Runtime/Synthesis/Source/SynthesisEditor/Classes/MonoWaveTablePresetBank.h
@@ -13,7 +13,7 @@ class FAssetTypeActions_MonoWaveTableSynthPreset : public FAssetTypeActions_Base
 {
 public:
 	virtual FText GetName() const override { return NSLOCTEXT("AssetTypeActions", "AssetTypeActions_MonoWaveTableSynthPreset", "Mono Wave Table Synth Preset Bank"); }
-	virtual FColor GetTypeColor() const override { return FColor(175, 100, 155); }
+	virtual FColor GetTypeColor() const override { return FColor(0, 255, 0); }
 	virtual UClass* GetSupportedClass() const override;
 	virtual uint32 GetCategories() override { return EAssetTypeCategories::Sounds; }
 };
diff --git a/Engine/Source/Editor/AudioEditor/Private/AudioEditorModule.cpp b/Engine/Source/Editor/AudioEditor/Private/AudioEditorModule.cpp
--- a/Engine/Source/Editor/AudioEditor/Private/AudioEditorModule.cpp
+++ b/Engine/Source/Editor/AudioEditor/Private/AudioEditorModule.cpp
@@ -309,7 +309,7 @@ private:
 		SET_AUDIO_ICON(SoundEffectSourcePreset, SourceEffectPreset);
 		SET_AUDIO_ICON(SoundEffectSourcePresetChain, SourceEffectPresetChain_1);
 		SET_AUDIO_ICON(ModularSynthPresetBank, SoundGenericIcon_2);
-		SET_AUDIO_ICON(MonoWaveTableSynthPresetBank, SoundGenericIcon_2);
+		SET_AUDIO_ICON(MonoWaveTableSynthPreset, SoundGenericIcon_2);
 		SET_AUDIO_ICON(TimeSynthClip, SoundGenericIcon_2);
 		SET_AUDIO_ICON(TimeSynthVolumeGroup, SoundGenericIcon_1);
 
