Fix missing include
#rnx

diff --git a/Engine/Plugins/Runtime/Synthesis/Source/Synthesis/Classes/SynthComponents/SynthComponentGranulator.h b/Engine/Plugins/Runtime/Synthesis/Source/Synthesis/Classes/SynthComponents/SynthComponentGranulator.h
--- a/Engine/Plugins/Runtime/Synthesis/Source/Synthesis/Classes/SynthComponents/SynthComponentGranulator.h
+++ b/Engine/Plugins/Runtime/Synthesis/Source/Synthesis/Classes/SynthComponents/SynthComponentGranulator.h
@@ -5,6 +5,7 @@
 #include "CoreMinimal.h"
 #include "Components/SynthComponent.h"
 #include "DSP/Granulator.h"
+#include "Sound/SampleBufferIO.h"
 #include "Sound/SoundWave.h"
 #include "SynthComponentGranulator.generated.h"
 
