Fix missing include
#rnx

diff --git a/Engine/Plugins/Runtime/Synthesis/Source/Synthesis/Classes/SynthComponents/SynthComponentWaveTable.h b/Engine/Plugins/Runtime/Synthesis/Source/Synthesis/Classes/SynthComponents/SynthComponentWaveTable.h
--- a/Engine/Plugins/Runtime/Synthesis/Source/Synthesis/Classes/SynthComponents/SynthComponentWaveTable.h
+++ b/Engine/Plugins/Runtime/Synthesis/Source/Synthesis/Classes/SynthComponents/SynthComponentWaveTable.h
@@ -6,6 +6,7 @@
 #include "Components/SynthComponent.h"
 #include "DSP/SampleBufferReader.h"
 #include "SampleBuffer.h"
+#include "Sound/SampleBufferIO.h"
 #include "Sound/SoundWave.h"
 #include "SynthComponentWaveTable.generated.h"
 
