Remove unused SSynthKnob tooltip which was causing crash in slate window render loop by creating window during construction, modifying the collection mid window draw iteration.

diff --git a/Engine/Plugins/Runtime/Synthesis/Source/Synthesis/Private/UI/SSynthKnob.cpp b/Engine/Plugins/Runtime/Synthesis/Source/Synthesis/Private/UI/SSynthKnob.cpp
--- a/Engine/Plugins/Runtime/Synthesis/Source/Synthesis/Private/UI/SSynthKnob.cpp
+++ b/Engine/Plugins/Runtime/Synthesis/Source/Synthesis/Private/UI/SSynthKnob.cpp
@@ -38,7 +38,6 @@ void SSynthKnob::Construct(const SSynthKnob::FArguments& InDeclaration)
 	bControllerInputCaptured = false;
 
 	// independently create a synth tooltip slate object (not a child)
-	SynthTooltip = SNew(SSynthTooltip);
 }
 
 int32 SSynthKnob::OnPaint(const FPaintArgs& Args, const FGeometry& AllottedGeometry, const FSlateRect& MyCullingRect, FSlateWindowElementList& OutDrawElements, int32 LayerId, const FWidgetStyle& InWidgetStyle, bool bParentEnabled) const
@@ -224,8 +223,6 @@ FReply SSynthKnob::OnMouseButtonUp(const FGeometry& MyGeometry, const FPointerEv
 		// Release capture for controller/keyboard when switching to mouse.
 		ResetControllerState();
 
-		SynthTooltip->SetWindowContainerVisibility(false);
-
 		return FReply::Handled().ReleaseMouseCapture();
 	}
 
@@ -252,20 +249,6 @@ FReply SSynthKnob::OnMouseMove(const FGeometry& MyGeometry, const FPointerEvent&
 		// Release capture for controller/keyboard when switching to mouse
 		ResetControllerState();
 
-		// TODO: fix tooltip
-// 		if (ShowTooltip.Get())
-// 		{
-// 			FVector2D TooltipPosition = MyGeometry.AbsolutePosition + MyGeometry.Size * MyGeometry.Scale + FVector2D(2.0f, 2.0f);
-// 			SynthTooltip->SetOverlayWindowPosition(TooltipPosition);
-// 			SynthTooltip->SetWindowContainerVisibility(true);
-// 
-// 			FNumberFormattingOptions FormattingOptions;
-// 			FText ValueNumberText = FText::AsNumber(NewValue, &FormattingOptions);
-// 			FText ValueText = FText::Format(LOCTEXT("SynthKnobValue", "{0} {1} {2}"), FText::FromString(TEXT("Frequency")), ValueNumberText, FText::FromString(TEXT("hz")));
-// 
-// 			SynthTooltip->SetOverlayText(ValueText);
-// 		}
-		
 		return FReply::Handled();
 	}
 
diff --git a/Engine/Plugins/Runtime/Synthesis/Source/Synthesis/Public/UI/SSynthKnob.h b/Engine/Plugins/Runtime/Synthesis/Source/Synthesis/Public/UI/SSynthKnob.h
--- a/Engine/Plugins/Runtime/Synthesis/Source/Synthesis/Public/UI/SSynthKnob.h
+++ b/Engine/Plugins/Runtime/Synthesis/Source/Synthesis/Public/UI/SSynthKnob.h
@@ -196,7 +196,6 @@ private:
 	bool bIsFocusable;
 
 private:
-	TSharedPtr<SSynthTooltip> SynthTooltip;
 
 	TArray<FSlateBrush*> KnobImages;
 
