Fix crash when attempting to set delay line beyond length of statically allocated max length

diff --git a/Engine/Plugins/Runtime/Synthesis/Source/Synthesis/Classes/SubmixEffects/SubmixEffectDelay.h b/Engine/Plugins/Runtime/Synthesis/Source/Synthesis/Classes/SubmixEffects/SubmixEffectDelay.h
--- a/Engine/Plugins/Runtime/Synthesis/Source/Synthesis/Classes/SubmixEffects/SubmixEffectDelay.h
+++ b/Engine/Plugins/Runtime/Synthesis/Source/Synthesis/Classes/SubmixEffects/SubmixEffectDelay.h
@@ -27,7 +27,7 @@ struct SYNTHESIS_API FSubmixEffectDelaySettings
 	UPROPERTY(EditAnywhere, BlueprintReadWrite, Category = Realtime, meta = (ClampMin = "0.0", UIMin = "0.0", UIMax = "20000.0"))
 	float InterpolationTime;
 
-	// Number of milliseconds of delay.
+	// Number of milliseconds of delay.  Caps at max delay at runtime.
 	UPROPERTY(EditAnywhere, BlueprintReadWrite, Category = Realtime, meta = (ClampMin = "0.0", UIMin = "0.0", UIMax = "20000.0"))
 	float DelayLength;
 
diff --git a/Engine/Plugins/Runtime/Synthesis/Source/Synthesis/Private/SubmixEffectDelay.cpp b/Engine/Plugins/Runtime/Synthesis/Source/Synthesis/Private/SubmixEffectDelay.cpp
--- a/Engine/Plugins/Runtime/Synthesis/Source/Synthesis/Private/SubmixEffectDelay.cpp
+++ b/Engine/Plugins/Runtime/Synthesis/Source/Synthesis/Private/SubmixEffectDelay.cpp
@@ -76,25 +76,29 @@ void FSubmixEffectDelay::SetInterpolationTime(float Time)
 
 void FSubmixEffectDelay::SetDelayLineLength(float Length)
 {
+	static const float MinLengthDelaySec = 0.4f;
+
 	TargetDelayLineLength = FMath::Clamp(Length, 0.4f, MaxDelayLineLength);
 	InterpolationInfo.SetValue(TargetDelayLineLength, InterpolationTime);
 }
 
 void FSubmixEffectDelay::UpdateParameters()
 {
+
 	FSubmixEffectDelaySettings NewSettings;
 
 	if (Params.GetParams(&NewSettings))
 	{
 		Audio::FDelay* DelaysPtr = DelayLines.GetData();
 
-		MaxDelayLineLength = NewSettings.MaximumDelayLength;
+		const float LastLength = MaxDelayLineLength;
+
+		MaxDelayLineLength = FMath::Max(NewSettings.MaximumDelayLength, MinLengthDelaySec);
 		InterpolationTime = NewSettings.InterpolationTime / 1000.0f;
 
-		TargetDelayLineLength = FMath::Clamp(NewSettings.DelayLength, 0.4f, MaxDelayLineLength);
-		InterpolationInfo.SetValue(TargetDelayLineLength, InterpolationTime);
+		SetDelayLineLength(NewSettings.DelayLength);
 
-		if (MaxDelayLineLength != NewSettings.MaximumDelayLength)
+		if (MaxDelayLineLength != LastLength)
 		{
 			for (int32 DelayIndex = 0; DelayIndex < DelayLines.Num(); DelayIndex++)
 			{
diff --git a/Engine/Source/Runtime/SignalProcessing/Private/Delay.cpp b/Engine/Source/Runtime/SignalProcessing/Private/Delay.cpp
--- a/Engine/Source/Runtime/SignalProcessing/Private/Delay.cpp
+++ b/Engine/Source/Runtime/SignalProcessing/Private/Delay.cpp
@@ -47,7 +47,7 @@ namespace Audio
 	{
 		if (AudioBuffer)
 		{
-			FMemory::Memzero(AudioBuffer, sizeof(float)*AudioBufferSize);
+			FMemory::Memzero(AudioBuffer, sizeof(float) * AudioBufferSize);
 		}
 
 		WriteIndex = 0;
@@ -59,8 +59,8 @@ namespace Audio
 	void FDelay::SetDelayMsec(const float InDelayMsec)
 	{
 		// Directly set the delay
-		DelayInSamples = InDelayMsec * SampleRate * 0.001f;
-		check(DelayInSamples < AudioBufferSize);
+		const float NewDelayInSamples = InDelayMsec * SampleRate * 0.001f;
+		DelayInSamples = FMath::Min(NewDelayInSamples, static_cast<float>(AudioBufferSize));
 		Update(true);
 	}
 
@@ -75,7 +75,8 @@ namespace Audio
 		EaseDelayMsec.SetValue(InDelayMsec, bIsInit);
 		if (bIsInit)
 		{
-			DelayInSamples = InDelayMsec * SampleRate * 0.001f;
+			const float NewDelayInSamples = InDelayMsec * SampleRate * 0.001f;
+			DelayInSamples = FMath::Min(NewDelayInSamples, static_cast<float>(AudioBufferSize));
 		}
 		Update(bIsInit);
 	}
