Merging Convolution Reverb

//UE4/Private-Reverb/Engine/Plugins/Runtime/Synthesis/...

to //UE4/Release-4.25/Engine/Plugins/Runtime/Synthesis/...

diff --git a/Engine/Plugins/Runtime/Synthesis/Source/Synthesis/Private/ConvolutionAlgorithmWrapper.cpp b/Engine/Plugins/Runtime/Synthesis/Source/Synthesis/Private/ConvolutionAlgorithmWrapper.cpp
new file mode 100644
--- /dev/null
+++ b/Engine/Plugins/Runtime/Synthesis/Source/Synthesis/Private/ConvolutionAlgorithmWrapper.cpp
@@ -0,0 +1,333 @@
+// Copyright Epic Games, Inc. All Rights Reserved.
+
+#include "ConvolutionAlgorithmWrapper.h"
+#include "Async/Async.h"
+#include "Misc/ScopeTryLock.h"
+
+namespace Audio
+{
+
+	void ConvolutionAlgorithmBuilder::DoWork()
+	{
+		check(AlgorithmWrapperPtr);
+		AlgorithmWrapperPtr->AsyncRebuildConvolutionObjectInternal();
+	}
+
+	// Default constructor. Object is not prepared to render audio
+	FConvolutionAlgorithmWrapper::FConvolutionAlgorithmWrapper()
+	{
+	}
+
+	FConvolutionAlgorithmWrapper::FConvolutionAlgorithmWrapper(const FConvolutionSettings & InSettings, AlignedFloatBuffer InitialIrData)
+		: CurrentSettings(InSettings)
+		, CachedDeinterleavedIR(InitialIrData)
+	{
+	}
+
+	FConvolutionAlgorithmWrapper::~FConvolutionAlgorithmWrapper()
+	{
+		if (TaskPtr)
+		{
+			TaskPtr->EnsureCompletion();
+			TaskPtr = nullptr;
+		}
+	}
+
+	void FConvolutionAlgorithmWrapper::ProcessAudio(int32 NumInputChannels, const AlignedFloatBuffer& InputInterleavedAudioBuffer, int32 NumOutputChannels, AlignedFloatBuffer& OutputInterleavedAudioBuffer)
+	{
+		const int32 NumFramesInThisCallback = InputInterleavedAudioBuffer.Num() / NumInputChannels;
+
+		// see if something changed between callbacks which requires updating the Convolution Algorithm callback
+		const bool bSomethingChanged = NumInputChannels != CurrentSettings.NumInputChannels		// input channel config changed...
+									|| NumOutputChannels != CurrentSettings.NumOutputChannels	// ...or output channel config changed
+									|| NumFramesInThisCallback != CurrentSettings.BlockNumSamples;			// ...or callback size changed
+
+		// not waiting for mutex
+		FScopeTryLock ScopeLock(&ChangingSettingsCritSec);
+
+		if (bSomethingChanged)
+		{
+			CurrentSettings.BlockNumSamples = NumFramesInThisCallback;
+			CurrentSettings.NumInputChannels = NumInputChannels;
+			CurrentSettings.NumOutputChannels = NumOutputChannels;
+
+			UpdateConvolutionObject(CurrentSettings);
+		}
+
+		// Passing-through audio to output: Either something changed or we couldn't obtain the mutex (dsp obj being rebuilt)
+		if (bSomethingChanged || !ScopeLock.IsLocked() || !CurrentConvolutionAlgorithm.IsValid())
+		{
+			const int32 MinNumChannels = FMath::Min(NumInputChannels, NumOutputChannels);
+			const int32 NumSamples = InputInterleavedAudioBuffer.Num();
+			const int32 NumFrames = NumSamples / NumInputChannels;
+			const float* InputPtr = InputInterleavedAudioBuffer.GetData();
+			float* OutputPtr = OutputInterleavedAudioBuffer.GetData();
+
+			if (NumInputChannels == NumOutputChannels)
+			{
+				FMemory::Memcpy(OutputPtr, InputPtr, NumSamples * sizeof(float));
+			}
+			else
+			{
+				// input and output channel config are different
+				// todo: better mix-down than channel to channel copy?
+				for (int32 CurrFrame = 0; CurrFrame < NumFrames; ++CurrFrame)
+				{
+					for (int32 CurrChan = 0; CurrChan < MinNumChannels; ++CurrChan)
+					{
+						OutputPtr[CurrFrame * NumOutputChannels + CurrChan] = InputPtr[CurrFrame * NumInputChannels + CurrChan];
+					}
+				}
+			}
+
+			return; // end pass-through
+		}
+
+		// Prepare scratch buffers (intermediate data is de-interleaved)
+		int32 NumInputSamples = InputInterleavedAudioBuffer.Num();
+		int32 NumOutputSamples = OutputInterleavedAudioBuffer.Num();
+
+		ScratchInputBuffer.Reset();
+		ScratchInputBuffer.AddDefaulted(NumInputSamples);
+
+		ScratchOutputBuffer.Reset();
+		ScratchOutputBuffer.AddDefaulted(NumOutputSamples);
+
+		// get pointers to buffers
+		const float* InputBufferPtr = InputInterleavedAudioBuffer.GetData();
+		float* OutputBufferPtr = OutputInterleavedAudioBuffer.GetData();
+		float* ScratchInputBufferPtr = ScratchInputBuffer.GetData();
+		float* ScratchOutputBufferPtr = ScratchOutputBuffer.GetData();
+
+		// De-interleave Input buffer into scratch input buffer
+		DeInterleaveBuffer(ScratchInputBufferPtr, InputBufferPtr, NumInputSamples, NumInputChannels);
+
+		// Convolution Algorithm object expects 2D arrays
+		GetAccessToBufferAs2DArray(ScratchInputBufferPtr, NumInputChannels, NumFramesInThisCallback, TempInput2DFloatArray);
+		GetAccessToBufferAs2DArray(ScratchOutputBufferPtr, NumOutputChannels, NumFramesInThisCallback, TempOuput2DFloatArray);
+
+		// get next buffer of convolution output (store in scratch output buffer)
+		CurrentConvolutionAlgorithm->ProcessAudioBlock(TempInput2DFloatArray.GetData(), TempOuput2DFloatArray.GetData());
+
+		// re-interleave scratch output buffer into final output buffer
+		InterleaveBuffer(OutputBufferPtr, ScratchOutputBufferPtr, NumOutputSamples, NumOutputChannels);
+	}
+
+	void FConvolutionAlgorithmWrapper::UpdateConvolutionObject(const FConvolutionSettings& InNewSettings, const AlignedFloatBuffer* OptionalInNewIrData)
+	{
+		// early exit if nothing is changing
+		if (CurrentSettings == InNewSettings && !OptionalInNewIrData)
+		{
+			return;
+		}
+
+		CurrentSettings = InNewSettings;
+		AsyncRebuildConvolutionObject(OptionalInNewIrData);
+	}
+
+
+#pragma region AudioBufferHelpers
+	void FConvolutionAlgorithmWrapper::GetAccessToBufferAs2DArray(float * InBufferPtr, const int32 NumChannels, const int32 NumFrames, TArray<float*>& Out2DArray)
+	{
+		// Obtain pointers to first sample in each de-interleaved channel
+		Out2DArray.Reset();
+		Out2DArray.AddUninitialized(NumChannels);
+
+		float** OutIrPtrs = Out2DArray.GetData();
+		int32 ChannelOffset = 0;
+
+		for (int32 i = 0; i < NumChannels; ++i)
+		{
+			OutIrPtrs[i] = (InBufferPtr + ChannelOffset);
+			ChannelOffset += NumFrames;
+		}
+	}
+
+	void FConvolutionAlgorithmWrapper::InterleaveBuffer(float* Dest, const float* Src, const int32 NumSamples, const int32 NumChannels)
+	{
+		const int32 NumFrames = NumSamples / NumChannels;
+		for (int32 i = 0; i < NumSamples; ++i)
+		{
+			const int32 CurrFrame = i / NumChannels;
+			const int32 CurrChannel = i % NumChannels;
+			const int32 DeinterleavedIndex = CurrChannel * NumFrames + CurrFrame;
+			Dest[i] = Src[DeinterleavedIndex];
+		}
+	}
+
+	void FConvolutionAlgorithmWrapper::DeInterleaveBuffer(float* Dest, const float* Src, const int32 NumSamples, const int32 NumChannels)
+	{
+		const int32 NumFrames = NumSamples / NumChannels;
+		for (int32 i = 0; i < NumSamples; ++i)
+		{
+			const int32 CurrFrame = i / NumChannels;
+			const int32 CurrChannel = i % NumChannels;
+			const int32 CurrOutputIndex = CurrChannel * NumFrames + CurrFrame;
+			Dest[CurrOutputIndex] = Src[i];
+		}
+	}
+
+	void FConvolutionAlgorithmWrapper::ReverseBufferInPlace(float* InOutBuffer, const int32 NumSamples, const int32 NumChannels, bool bIsInterleaved)
+	{
+		if (bIsInterleaved)
+		{
+			FConvolutionAlgorithmWrapper::ReverseInterleavedBufferInPlace(InOutBuffer, NumSamples, NumChannels);
+		}
+		else
+		{
+			FConvolutionAlgorithmWrapper::ReverseDeinterleavedBufferInPlace(InOutBuffer, NumSamples, NumChannels);
+		}
+	}
+
+	void FConvolutionAlgorithmWrapper::ReverseInterleavedBufferInPlace(float* InOutBuffer, const int32 NumSamples, const int32 NumChannels)
+	{
+		const int32 NumFrames = NumSamples / NumChannels;
+		const int32 HalfNumFrames = NumFrames >> 1;
+		const int32 FrameSizeInBytes = NumChannels * sizeof(float);
+
+		float* BufferStart = InOutBuffer;
+		float* BufferEnd = BufferStart + (NumFrames - 1) * NumChannels;
+
+		for (int32 i = 0; i < HalfNumFrames; ++i)
+		{
+			FMemory::Memswap(BufferStart, BufferEnd, FrameSizeInBytes);
+			BufferStart += NumChannels;
+			BufferEnd -= NumChannels;
+		}
+	}
+
+	void FConvolutionAlgorithmWrapper::ReverseDeinterleavedBufferInPlace(float* InOutBuffer, const int32 NumSamples, const int32 NumChannels)
+	{
+		const int32 NumFrames = NumSamples / NumChannels;
+		const int32 HalfNumFrames = NumFrames >> 1;
+
+		for (int32 Chan = 0; Chan < NumChannels; ++Chan)
+		{
+			float* ChanStart = InOutBuffer + (Chan * NumFrames);
+			float* ChanEnd = ChanStart + (NumFrames - 1);
+
+			for (int32 i = 0; i < HalfNumFrames; ++i)
+			{
+				FMemory::Memswap(ChanStart++, ChanEnd--, sizeof(float));
+			}
+		}
+	}
+#pragma endregion
+
+#pragma region Getters
+	void FConvolutionAlgorithmWrapper::Reset()
+	{
+		FScopeLock ScopeLock(&ChangingSettingsCritSec);
+
+		if (CurrentConvolutionAlgorithm)
+		{
+			CurrentConvolutionAlgorithm = nullptr;
+		}
+	}
+
+	bool Audio::FConvolutionAlgorithmWrapper::IsHardwareAccelerationAllowed() const
+	{
+		return CurrentSettings.bEnableHardwareAcceleration;
+	}
+
+	int32 Audio::FConvolutionAlgorithmWrapper::GetBlockSizeInSamples() const
+	{
+		return CurrentSettings.BlockNumSamples;
+	}
+
+	int32 Audio::FConvolutionAlgorithmWrapper::GetNumInputChannels() const
+	{
+		return CurrentSettings.NumInputChannels;
+	}
+
+	int32 FConvolutionAlgorithmWrapper::GetNumOutputChannels() const
+	{
+		return CurrentSettings.NumOutputChannels;
+	}
+
+	int32 FConvolutionAlgorithmWrapper::GetNumImpulseResponseChannels() const
+	{
+		return CurrentSettings.NumImpulseResponses;
+	}
+
+	int32 FConvolutionAlgorithmWrapper::GetMaxNumImpulseResponseChannels() const
+	{
+		return CurrentSettings.MaxNumImpulseResponseSamples;
+	}
+#pragma endregion
+
+	void FConvolutionAlgorithmWrapper::AsyncRebuildConvolutionObject(const AlignedFloatBuffer* OptionalInNewIrData)
+	{
+		FScopeLock ScopeLock(&ChangingSettingsCritSec);
+		// if we have new IR data, we need to make our own copy for processing
+		// if we have no new IR data, and the channel config on the new data doesn't match,
+		//	we need to shutdown the current algorithm until the new IR is ready
+		if (OptionalInNewIrData)
+		{
+			CachedDeinterleavedIR = *OptionalInNewIrData; // 'spensive :(
+		}
+		else
+		{
+			Reset();
+		}
+
+		if (!CurrentSettingsAreRenderable())
+		{
+			return;
+		}
+
+		// start async task
+		if (TaskPtr)
+		{
+			TaskPtr->EnsureCompletion();
+		}
+
+		TaskPtr = MakeUnique<FAsyncTask<ConvolutionAlgorithmBuilder>>(this);
+		TaskPtr->StartBackgroundTask();
+	}
+
+	void FConvolutionAlgorithmWrapper::AsyncRebuildConvolutionObjectInternal()
+	{
+		FScopeLock ScopeLock(&ChangingSettingsCritSec);
+
+		const int32 NumIrFrames = CachedDeinterleavedIR.Num() / CurrentSettings.NumImpulseResponses;
+		const int32 NumIrChannels = CurrentSettings.NumImpulseResponses;
+
+		CurrentConvolutionAlgorithm = FConvolutionFactory::NewConvolutionAlgorithm(CurrentSettings);
+
+		// If fails, prior creation function warns why algorithm failed to be created.
+		if (CurrentConvolutionAlgorithm.IsValid())
+		{
+			if (CachedDeinterleavedIR.Num())
+			{
+				GetAccessToBufferAs2DArray(CachedDeinterleavedIR.GetData(), NumIrChannels, NumIrFrames, TempInput2DFloatArray);
+
+				check(TempInput2DFloatArray.Num() == CurrentSettings.NumImpulseResponses);
+
+				// Hand over impulse responses to convolution reverb object
+				for (int32 i = 0; i < CurrentSettings.NumImpulseResponses; ++i)
+				{
+					CurrentConvolutionAlgorithm->SetImpulseResponse(i, TempInput2DFloatArray[i], CurrentSettings.MaxNumImpulseResponseSamples);
+				}
+
+				// Set up (identity) 3D gain matrix
+				// todo: Make the gain matrix accessible to the sound designer in a meaningful way
+				const int32 MinNumChannels = FMath::Min3(CurrentSettings.NumInputChannels, CurrentSettings.NumOutputChannels, CurrentSettings.NumImpulseResponses);
+
+				for (int i = 0; i < MinNumChannels; ++i)
+				{
+					CurrentConvolutionAlgorithm->SetMatrixGain(i, i, i, 1.0f);
+				}
+			}
+		}
+	}
+
+	bool FConvolutionAlgorithmWrapper::CurrentSettingsAreRenderable()
+	{
+		return CurrentSettings.BlockNumSamples > 0
+			&& CurrentSettings.MaxNumImpulseResponseSamples > 0
+			&& CurrentSettings.NumImpulseResponses > 0
+			&& CurrentSettings.NumInputChannels > 0
+			&& CurrentSettings.NumOutputChannels > 0;
+	}
+} // namespace Audio
\ No newline at end of file
diff --git a/Engine/Plugins/Runtime/Synthesis/Source/Synthesis/Private/ConvolutionReverbComponent.cpp b/Engine/Plugins/Runtime/Synthesis/Source/Synthesis/Private/ConvolutionReverbComponent.cpp
new file mode 100644
--- /dev/null
+++ b/Engine/Plugins/Runtime/Synthesis/Source/Synthesis/Private/ConvolutionReverbComponent.cpp
@@ -0,0 +1,227 @@
+// Copyright Epic Games, Inc. All Rights Reserved.
+
+#pragma once
+#include "ConvolutionReverbComponent.h"
+#include "DSP/BufferVectorOperations.h"
+#include "DSP/ConvolutionAlgorithm.h"
+#include "DSP/Dsp.h"
+
+FConvolutionReverbSettings::FConvolutionReverbSettings()
+	: ImpulseResponse(nullptr)
+	, AllowHardwareAcceleration(true)
+	, PostNormalizationVolume_decibels(-24.0f)
+	, SurroundRearChannelBleedAmount(0.0f)
+	, bSurroundRearChannelFlip(false)
+{
+}
+
+FConvolutionReverb::FConvolutionReverb()
+	: CachedIrAssetPtr(nullptr)
+	, bIrIsDirty(false)
+	, bSettingsAreDirty(false)
+	, bIrBeingRebuilt(false)
+	, SampleRate(0.0f)
+	, CachedOutputGain(1.0f)
+	, RearChanBleed(0.0f)
+	, bRearChanFlip(false)
+	, bCachedEnableHardwareAcceleration(true)
+{
+	
+}
+
+FConvolutionReverb::~FConvolutionReverb()
+{
+	UImpulseResponse* IrPtr = CachedIrAssetPtr.Get();
+	if (IrPtr)
+	{
+		IrPtr->UnsubscribeFlagFromIrChanges(&bIrIsDirty);
+	}
+}
+
+void FConvolutionReverb::Init(const FSoundEffectSubmixInitData& InitData)
+{
+	// Do any initializations here
+	SampleRate = InitData.SampleRate;
+	bSettingsAreDirty = true;
+}
+
+void FConvolutionReverb::OnPresetChanged()
+{
+	GET_EFFECT_SETTINGS(ConvolutionReverb);
+	bIrBeingRebuilt = false;
+
+	// SURROUND
+	RearChanBleed = Settings.SurroundRearChannelBleedAmount;
+	bRearChanFlip = Settings.bSurroundRearChannelFlip;
+
+	// OUTPUT GAIN
+	const float NewOutputGain = Audio::ConvertToLinear(Settings.PostNormalizationVolume_decibels);
+	if (!FMath::IsNearlyEqual(CachedOutputGain, NewOutputGain))
+	{
+		CachedOutputGain = NewOutputGain;
+		bSettingsAreDirty = true;
+	}
+
+	// HARDWARE ACCELERATION
+	if (bCachedEnableHardwareAcceleration != Settings.AllowHardwareAcceleration)
+	{
+		bCachedEnableHardwareAcceleration = Settings.AllowHardwareAcceleration;
+		bSettingsAreDirty = true;
+	}
+
+
+	// IR ASSET UPDATE
+	UImpulseResponse* NewIrAssetPtr = Settings.ImpulseResponse;
+	UImpulseResponse* IrAssetPtrOld = CachedIrAssetPtr.Get();
+
+	// IR Asset hasn't been changed
+	if (NewIrAssetPtr == IrAssetPtrOld)
+	{
+		return;
+	}
+	// IR Asset was removed/deleted
+	else if (!NewIrAssetPtr)
+	{
+		if (IrAssetPtrOld)
+		{
+			IrAssetPtrOld->UnsubscribeFlagFromIrChanges(&bIrIsDirty);
+		}
+		ConvolutionAlgorithm.Reset();
+		CachedIrAssetPtr = NewIrAssetPtr;
+	}
+	else // received new IR Asset, swap subscriptions and update weak ptr
+	{
+		bIrIsDirty.AtomicSet(false);
+		if (IrAssetPtrOld)
+		{
+			IrAssetPtrOld->UnsubscribeFlagFromIrChanges(&bIrIsDirty);
+		}
+
+		NewIrAssetPtr->SubscribeFlagToIrChanges(&bIrIsDirty);
+
+		CachedIrAssetPtr = MakeWeakObjectPtr(NewIrAssetPtr);
+
+		// will prepare IR data async if its not already ready.
+		// IR asset will set to true when the data is ready
+		bIrBeingRebuilt = true;
+		NewIrAssetPtr->PrepareIRsForConvolutionAlgorithm(SampleRate);
+	}
+}
+
+void FConvolutionReverb::OnProcessAudio(const FSoundEffectSubmixInputData& InData, FSoundEffectSubmixOutputData& OutData)
+{
+	// Here, we get min of num channels of input and output (note that submix processing output channels can be different than input)
+	const int32 NumChannels = FMath::Min(InData.NumChannels, OutData.NumChannels);
+
+	const float* InputBuffer = InData.AudioBuffer->GetData();
+	float* OutputBuffer = OutData.AudioBuffer->GetData();
+
+	const int32 InputBufferSize = InData.AudioBuffer->Num();
+	const int32 OutputBufferSize = OutData.AudioBuffer->Num();
+	const int32 NumFrames = InData.NumFrames;
+	const int32 NumInputChannels = InData.NumChannels;
+	const int32 NumOutputChannels = OutData.NumChannels;
+	const int32 NumInputSamples = InData.AudioBuffer->Num();
+	const int32 NumOutputSamples = OutData.AudioBuffer->Num();
+
+
+	// Do we need to update the Convolution object?
+	// the IR Asset itself may have raised bIrIsDirty flag if it was async prepping the IR data
+	if (bSettingsAreDirty || (bIrBeingRebuilt && bIrIsDirty))
+	{
+		// get IR data
+		if (auto IrPtr = CachedIrAssetPtr.Get())
+		{
+			const Audio::AlignedFloatBuffer& IrDataRef = IrPtr->GetDeinterleavedIRData();
+
+			// build new settings struct
+			Audio::FConvolutionSettings NewSettings;
+
+			const int32 NumIrChannels = IrPtr->GetNumChannels();
+			NewSettings.bEnableHardwareAcceleration = bCachedEnableHardwareAcceleration;
+			NewSettings.BlockNumSamples = NumFrames;
+			NewSettings.MaxNumImpulseResponseSamples = IrDataRef.Num() / NumIrChannels;
+			NewSettings.NumImpulseResponses = NumIrChannels;
+			NewSettings.NumInputChannels = NumInputChannels;
+			NewSettings.NumOutputChannels = NumOutputChannels;
+
+			// update settings and IR data
+			if (bIrBeingRebuilt && bIrIsDirty)
+			{
+				ConvolutionAlgorithm.UpdateConvolutionObject(NewSettings, &IrDataRef);
+
+				bIrIsDirty.AtomicSet(false);
+				bIrBeingRebuilt = false;
+				bSettingsAreDirty = false;
+			}
+			else
+			{
+				// Update settings only
+				ConvolutionAlgorithm.UpdateConvolutionObject(NewSettings);
+				bSettingsAreDirty = false;
+			}
+		}
+	}
+
+
+	// generate output audio
+	ConvolutionAlgorithm.ProcessAudio(NumInputChannels, *InData.AudioBuffer, NumOutputChannels, *OutData.AudioBuffer);
+
+	const bool bIsSurroundOutput = (NumOutputChannels >= 4);
+	const bool bHasSurroundBleed = (!FMath::IsNearlyZero(RearChanBleed));
+	const float G = FMath::Sqrt(1.0f/2.0f);
+
+
+	float FinalOutputGain = CachedOutputGain;
+	
+	// Modify the output level if we are using surround bleed.
+	// The gain scalar is applied to compensate for louder output
+	// if our surround output signal is folded down to stereo
+	if (bIsSurroundOutput && bHasSurroundBleed && bRearChanFlip)
+	{
+		FinalOutputGain *= G;
+	}
+	else if (bIsSurroundOutput && bHasSurroundBleed)
+	{
+		FinalOutputGain *= 0.5f;
+	}
+
+	Audio::MultiplyBufferByConstantInPlace(*OutData.AudioBuffer, FinalOutputGain);
+
+	// surround bleed?
+	if (bIsSurroundOutput && bHasSurroundBleed)
+	{
+		float* OutPtr = OutData.AudioBuffer->GetData();
+
+		const int32 OutputFrameOffset = NumOutputChannels - 2;
+
+		if (!bRearChanFlip)
+		{
+			for (int32 i = 0; i < NumOutputSamples; i += NumOutputChannels)
+			{
+				const int32 LeftRear = i + OutputFrameOffset;
+				OutPtr[LeftRear] = RearChanBleed * OutPtr[i];
+				OutPtr[LeftRear + 1] = RearChanBleed * OutPtr[i + 1];
+			}
+		}
+		else
+		{
+			for (int32 i = 0; i < NumOutputSamples; i += NumOutputChannels)
+			{
+				const int32 LeftRear = i + OutputFrameOffset;
+				OutPtr[LeftRear + 1] = RearChanBleed * OutPtr[i];
+				OutPtr[LeftRear] = RearChanBleed * OutPtr[i + 1];
+			}
+		}
+	}
+}
+
+uint32 FConvolutionReverb::GetDesiredInputChannelCountOverride() const
+{
+	return 2;
+}
+
+void UConvolutionReverbPreset::SetSettings(const FConvolutionReverbSettings& InSettings)
+{
+	UpdateSettings(InSettings);
+}
diff --git a/Engine/Plugins/Runtime/Synthesis/Source/Synthesis/Private/ImpulseResponseAsset.cpp b/Engine/Plugins/Runtime/Synthesis/Source/Synthesis/Private/ImpulseResponseAsset.cpp
new file mode 100644
--- /dev/null
+++ b/Engine/Plugins/Runtime/Synthesis/Source/Synthesis/Private/ImpulseResponseAsset.cpp
@@ -0,0 +1,188 @@
+// Copyright Epic Games, Inc. All Rights Reserved.
+
+#include "ImpulseResponseAsset.h"
+#include "Async/Async.h"
+
+
+void ImpulseResponesConverter::DoWork()
+{
+	IrAssetPtr->PrepareIRsForConvolutionAlgorithmInternal(TargetSampleRate);
+	IrAssetPtr->SetSubscriberFlags(true);
+}
+
+
+UImpulseResponse::UImpulseResponse()
+: TaskPtr{ nullptr }
+, NumChannels(0)
+, NumFrames(0)
+, AssetSampleRate(0)
+{
+}
+
+UImpulseResponse::~UImpulseResponse()
+{
+	// stop our task if its in flight
+	if (TaskPtr)
+	{
+		TaskPtr->EnsureCompletion();
+	}
+
+	// set all subscriber flags to false
+	// (this is probably unnecessary, since they accesses this object via WeakPtr)
+	SetSubscriberFlags(false);
+}
+
+void UImpulseResponse::SubscribeFlagToIrChanges(FThreadSafeBool* ExternalThreadSafeBool)
+{
+	FScopeLock Lock(&SubscriberFlagsCritSec);
+
+	FThreadSafeBool** Flags = SubscriberFlags.GetData();
+
+	// try to replace null entry in array
+	const int32 NumFlags = SubscriberFlags.Num();
+	for (int32 i = 0; i < NumFlags; ++i)
+	{
+		if (!Flags[i])
+		{
+			Flags[i] = ExternalThreadSafeBool;
+			return;
+		}
+	}
+
+	// array was full
+	SubscriberFlags.Add(ExternalThreadSafeBool);
+}
+
+void UImpulseResponse::UnsubscribeFlagFromIrChanges(FThreadSafeBool* ExternalThreadSafeBool)
+{
+	FScopeLock Lock(&SubscriberFlagsCritSec);
+
+	FThreadSafeBool** Flags = SubscriberFlags.GetData();
+
+	// search for entry in array
+	const int32 NumFlags = SubscriberFlags.Num();
+	for (int32 i = 0; i < NumFlags; ++i)
+	{
+		if (Flags[i] == ExternalThreadSafeBool)
+		{
+			Flags[i] = nullptr;
+			return;
+		}
+	}
+
+	// TODO: Log wargning (trying to unsubscribe flag that was never subscribed)
+}
+
+void UImpulseResponse::SetSubscriberFlags(bool bInDataIsReady)
+{
+	FScopeLock Lock(&SubscriberFlagsCritSec);
+
+	// This access should be safe, since mutation of the Map is locked by a mutex
+	// and subscribers should call UnsubscribeFlagFromIrChanges() in their destructor
+	for (FThreadSafeBool* TSBool : SubscriberFlags)
+	{
+		if (TSBool)
+		{
+			TSBool->AtomicSet(bInDataIsReady);
+		}
+	}
+}
+
+void UImpulseResponse::PrepareIRsForConvolutionAlgorithm(const int32 InSampleRate)
+{
+	if (IsDataPrepared(InSampleRate))
+	{
+		SetSubscriberFlags(true);
+	}
+	else
+	{
+		// if the task is in flight, let it finish
+		if (TaskPtr)
+		{
+			TaskPtr->EnsureCompletion();
+		}
+
+		TaskPtr = MakeUnique<FAsyncTask<ImpulseResponesConverter>>(this, InSampleRate);
+		TaskPtr->StartBackgroundTask();
+	}
+}
+
+void UImpulseResponse::PrepareIRsForConvolutionAlgorithmInternal(const int32 InSampleRate)
+{
+	const float StepSize = static_cast<float>(InSampleRate) / AssetSampleRate;
+	const int32 FinalBufferSize = static_cast<int32>(StepSize * IRData.Num());
+
+	if ((DeinterleavedIR.Num() != FinalBufferSize) || (InSampleRate != AssetSampleRate))
+	{
+		PerformAssetIrSRC(InSampleRate);
+		CurrentSampleRate = InSampleRate;
+	}
+}
+
+bool UImpulseResponse::IsDataPrepared(const int32 InSampleRate)
+{
+	// Are we preparing the data asynchronously now?
+	if (TaskPtr && !TaskPtr->IsDone())
+	{
+		return false;
+	}
+	else if (TaskPtr) // we are holding on to a finished task
+	{
+		TaskPtr.Reset();
+	}
+
+	// Is the current data already prepared
+	return (CurrentSampleRate == InSampleRate) && (IRData.Num() == DeinterleavedIR.Num());
+}
+
+void UImpulseResponse::PerformAssetIrSRC(int32 InSampleRate)
+{
+	const float StepSize = AssetSampleRate / static_cast<float>(InSampleRate);
+	const int32 AssetBufferSize = IRData.Num();
+	const int32 FinalBufferSize = static_cast<int32>(StepSize * AssetBufferSize);
+	const int32 FinalNumFrames = static_cast<int32>(StepSize * NumFrames);
+	float SampleA, SampleB;
+
+	DeinterleavedIR.Reset();
+	DeinterleavedIR.AddUninitialized(FinalBufferSize);
+
+	if (InSampleRate == static_cast<int32>(AssetSampleRate))
+	{
+		// no src needed, only copy
+		FMemory::Memcpy(DeinterleavedIR.GetData(), IRData.GetData(), IRData.Num() * sizeof(float));
+		return;
+	}
+
+	// SRC: (need to test this with an IR that has a different SR than the FAudioDevice)
+	// for each channel...
+	for (int32 Chan = 0; Chan < NumChannels; ++Chan)
+	{
+		const int32 InputChanOffset = NumFrames * Chan;
+		const int32 OutputChanOffset = FinalNumFrames * Chan;
+
+		// for each output sample...
+		float CurrIndex = 0.0f;
+		for (int32 i = 0; i < FinalNumFrames; ++i)
+		{
+			const float Alpha = FMath::Fmod(CurrIndex, 1.0f);
+			const int32 WholeThisIndex = FMath::FloorToInt(CurrIndex);
+			int32 WholeNextIndex = WholeThisIndex + 1;
+
+			// check for interpolation between last and first frames
+			SampleA = IRData[InputChanOffset + WholeThisIndex];
+
+			if (WholeNextIndex != AssetBufferSize)
+			{
+				SampleB = IRData[InputChanOffset + WholeNextIndex];
+			}
+			else
+			{
+				SampleB = 0.0f;
+			}
+
+			DeinterleavedIR[OutputChanOffset + i] = FMath::Lerp(SampleA, SampleB, Alpha);
+
+			CurrIndex += StepSize;
+		}
+	}
+}
\ No newline at end of file
diff --git a/Engine/Plugins/Runtime/Synthesis/Source/Synthesis/Public/ConvolutionAlgorithmWrapper.h b/Engine/Plugins/Runtime/Synthesis/Source/Synthesis/Public/ConvolutionAlgorithmWrapper.h
new file mode 100644
--- /dev/null
+++ b/Engine/Plugins/Runtime/Synthesis/Source/Synthesis/Public/ConvolutionAlgorithmWrapper.h
@@ -0,0 +1,140 @@
+// Copyright Epic Games, Inc. All Rights Reserved.
+#pragma once
+
+
+#include "CoreMinimal.h"
+#include "Components/AudioComponent.h"
+#include "DSP/ConvolutionAlgorithm.h"
+
+namespace Audio
+{
+
+	class FConvolutionAlgorithmWrapper;
+
+
+	// ========================================================================
+	// FConvolutionAlgorithmBuilder:
+	//
+	// Async Task class used by FConvolutionAlgorithmWrapper.
+	// Giving new IRs to a convolution algorithm object can be expensive (FFT and copies)
+	// This allows us to do this work asynchronously
+	// ========================================================================
+	class SYNTHESIS_API ConvolutionAlgorithmBuilder : public FNonAbandonableTask
+	{
+	public:
+		ConvolutionAlgorithmBuilder(FConvolutionAlgorithmWrapper* InAlgorithmWrapperPtr)
+			: AlgorithmWrapperPtr(InAlgorithmWrapperPtr)
+		{
+		}
+
+	private:
+		void DoWork();
+
+		FConvolutionAlgorithmWrapper* AlgorithmWrapperPtr;
+
+		FORCEINLINE TStatId GetStatId() const
+		{
+			RETURN_QUICK_DECLARE_CYCLE_STAT(ConvolutionAlgorithmBuilder, STATGROUP_ThreadPoolAsyncTasks);
+		}
+
+		friend class FAsyncTask<ConvolutionAlgorithmBuilder>;
+	}; // class ConvolutionAlgorithmBuilder
+
+
+	// ========================================================================
+	// FConvolutionAlgorithmWrapper:
+	//
+	// Wrapper around IConvolutionAlgorithm
+	// Automatically handles re-building the underlying ConvolutionAlgorithm object asynchronosuly
+	// Deals with thread saftey (methods are safe to call from any thread)
+	// ========================================================================
+	class SYNTHESIS_API FConvolutionAlgorithmWrapper
+	{
+	public:
+		FConvolutionAlgorithmWrapper();
+
+		FConvolutionAlgorithmWrapper(const FConvolutionSettings& InSettings, AlignedFloatBuffer InitialIrData);
+
+		virtual ~FConvolutionAlgorithmWrapper();
+
+		// meant to be called from audio render thread
+		virtual void ProcessAudio(int32 NumInputChannels, const AlignedFloatBuffer& InputInterleavedAudioBuffer, int32 NumOutputChannels, AlignedFloatBuffer& OutputInterleavedAudioBuffer);
+
+		// Can be called on any thread. Will do heavy lifting on Game Thread (ProcessAudio will perform pass-through in the meantime)
+		void UpdateConvolutionObject(const FConvolutionSettings& InNewSettings, const AlignedFloatBuffer* OptionalInNewIrDataPtr = nullptr);
+
+	#pragma region Getters
+		void Reset();
+
+		// Hardware Acceleration
+		bool IsHardwareAccelerationAllowed() const;
+
+		// block size
+		int32 GetBlockSizeInSamples() const;
+
+		// Number of input channels
+		int32 GetNumInputChannels() const;
+
+		// Number of output channels
+		int32 GetNumOutputChannels() const;
+
+		// Number of impulse response channels
+		int32 GetNumImpulseResponseChannels() const;
+
+		// maximum number of impulse response channels
+		int32 GetMaxNumImpulseResponseChannels() const;
+	#pragma endregion
+	
+	private:
+		// spawns a task on the game thread. bIsPreparedToRenderAudio will be false until the new convolution object is ready
+		void AsyncRebuildConvolutionObject(const AlignedFloatBuffer* OptionalInNewIrDataPtr = nullptr);
+	
+		void AsyncRebuildConvolutionObjectInternal();
+	
+		bool CurrentSettingsAreRenderable();
+
+	#pragma region AudioBufferHelpers
+		// audio buffer-related helpers
+		static void GetAccessToBufferAs2DArray(float* InBufferPtr, const int32 NumChannels, const int32 NumFrames, TArray<float*>& Out2DArray);
+
+		static void InterleaveBuffer(float* Dest, const float* Src, const int32 NumSamples, const int32 NumChannels);
+
+		static void DeInterleaveBuffer(float* Dest, const float* Src, const int32 NumSamples, const int32 NumChannels);
+
+		static void ReverseBufferInPlace(float* InOutBuffer, const int32 NumSamples, const int32 NumChannels, bool bIsInterleaved);
+
+		static void ReverseInterleavedBufferInPlace(float* InOutBuffer, const int32 NumSamples, const int32 NumChannels);
+
+		static void ReverseDeinterleavedBufferInPlace(float* InOutBuffer, const int32 NumSamples, const int32 NumChannels);
+	#pragma endregion
+
+		TUniquePtr<IConvolutionAlgorithm> CurrentConvolutionAlgorithm;
+
+		TUniquePtr<FAsyncTask<ConvolutionAlgorithmBuilder>> TaskPtr;
+
+		FConvolutionSettings CurrentSettings;
+
+		// cache this so we can rebuild the algorithm (for non-IR reasons) without depending on a UObject
+		AlignedFloatBuffer CachedDeinterleavedIR;
+
+		// data is passed to the convolution algorithm as 2D arrays
+		TArray<float*> TempInput2DFloatArray;
+
+		// data is recieved from the convolution algorithm as 2D arrays
+		TArray<float*> TempOuput2DFloatArray;
+
+		// scratch buffer used to de-interleave input audio
+		AlignedFloatBuffer ScratchInputBuffer;
+
+		// scratch buffer used to interleave output audio
+		AlignedFloatBuffer ScratchOutputBuffer;
+
+		// critical section used to protect audio rendering
+		FCriticalSection ChangingSettingsCritSec;
+
+		// allow async worker to rebuild dsp object
+		friend class ConvolutionAlgorithmBuilder;
+
+	}; // class FConvolutionAlgorithmWrapper
+
+} // namespace Audio
diff --git a/Engine/Plugins/Runtime/Synthesis/Source/Synthesis/Public/ConvolutionReverbComponent.h b/Engine/Plugins/Runtime/Synthesis/Source/Synthesis/Public/ConvolutionReverbComponent.h
new file mode 100644
--- /dev/null
+++ b/Engine/Plugins/Runtime/Synthesis/Source/Synthesis/Public/ConvolutionReverbComponent.h
@@ -0,0 +1,118 @@
+// Copyright Epic Games, Inc. All Rights Reserved.
+#pragma once
+
+#include "CoreMinimal.h"
+#include "Components/AudioComponent.h"
+#include "UObject/ObjectMacros.h"
+#include "UObject/Object.h"
+#include "UObject/Class.h"
+#include "Engine/EngineTypes.h"
+#include "Engine/DataTable.h"
+#include "Sound/SoundEffectSubmix.h"
+#include "ImpulseResponseAsset.h"
+#include "AudioEffect.h"
+#include "Kismet/BlueprintFunctionLibrary.h"
+#include "ConvolutionAlgorithmWrapper.h"
+#include "ConvolutionReverbComponent.generated.h"
+
+ 
+// ========================================================================
+// FConvolutionReverbSettings:
+// 
+// settings component for FConvolutionReverb submix effect
+// ========================================================================
+
+
+USTRUCT(BlueprintType)
+struct SYNTHESIS_API FConvolutionReverbSettings
+{
+	GENERATED_USTRUCT_BODY()
+
+	FConvolutionReverbSettings();
+
+	/* Information about your parameter described here will show up in the editor as a tool-tip in the effect preset object */
+	UPROPERTY(EditAnywhere, BlueprintReadWrite, Category = SoundEffect)
+	UImpulseResponse* ImpulseResponse;
+
+	/* Opt into hardware acceleration of the convolution reverb (if available) */
+	UPROPERTY(EditAnywhere, BlueprintReadWrite, Category = SoundEffect)
+	bool AllowHardwareAcceleration;
+
+	/* Used to account for energy added by convolution with "loud" Impulse Responses.  Not meant to be updated dynamically */
+	UPROPERTY(EditAnywhere, BlueprintReadWrite, Category = SoundEffect)
+	float PostNormalizationVolume_decibels;
+
+	/* Amout of audio to be sent to rear channels in quad/surround configurations (linear gain, < 0 = phase inverted) */
+	UPROPERTY(EditAnywhere, BlueprintReadWrite, Category = SubmixEffectPreset, meta = (ClampMin = "-5.0", UIMin = "-5.0", ClampMax = "5.0", UIMax = "5.0"))
+	float SurroundRearChannelBleedAmount;
+
+	/* If true, send Surround Rear Channel Bleed Amount sends front left to back right and vice versa */
+	UPROPERTY(EditAnywhere, BlueprintReadWrite, Category = SubmixEffectPreset)
+	bool bSurroundRearChannelFlip;
+};
+
+// ========================================================================
+// UConvolutionReverbPreset:
+//
+// Class which processes audio streams and uses parameters defined in the preset class.
+// This class is instanced for each submix effect preset in a submix effect chain.
+// ========================================================================
+
+class SYNTHESIS_API FConvolutionReverb : public FSoundEffectSubmix
+{
+public:
+	FConvolutionReverb();
+	~FConvolutionReverb();
+
+	// Being FSoundEffectSubmix Implementation
+	// Called on an audio effect at initialization on main thread before audio processing begins.
+	virtual void Init(const FSoundEffectSubmixInitData& InitData) override;
+
+	// Called when an audio effect preset is changed
+	virtual void OnPresetChanged() override;
+
+	// Process the input block of audio. Called on audio thread.
+	virtual void OnProcessAudio(const FSoundEffectSubmixInputData& InData, FSoundEffectSubmixOutputData& OutData) override;
+
+	// We want to receive down-mixed submix audio to stereo input for the reverb effect
+	virtual uint32 GetDesiredInputChannelCountOverride() const override;
+	// End FSoundEffectSubmix Implementation
+
+private:
+	Audio::FConvolutionAlgorithmWrapper ConvolutionAlgorithm;
+	TArray<float> ScratchInputBuffer;
+	TArray<float> ScratchOutputBuffer;
+	TWeakObjectPtr<UImpulseResponse> CachedIrAssetPtr;
+	FThreadSafeBool bIrIsDirty; // need to rebuild ConvolutionAlgorithm if true
+	bool bSettingsAreDirty;
+	bool bIrBeingRebuilt;
+	float SampleRate;
+
+	// members to update in OnPresetChanged()
+	float CachedOutputGain;
+	float RearChanBleed;
+	bool bRearChanFlip;
+	bool bCachedEnableHardwareAcceleration;
+};
+
+// ========================================================================
+// UConvolutionReverbPreset
+// This code exposes your preset settings and effect class to the editor.
+// ========================================================================
+
+UCLASS()
+class SYNTHESIS_API UConvolutionReverbPreset : public USoundEffectSubmixPreset
+{
+	GENERATED_BODY()
+
+public:
+	EFFECT_PRESET_METHODS(ConvolutionReverb)
+
+	UFUNCTION(BlueprintCallable, Category = "Audio|Effects")
+	void SetSettings(const FConvolutionReverbSettings& InSettings);
+
+	// ConvolutionReverbPreset Preset Settings
+	UPROPERTY(EditAnywhere, BlueprintReadWrite, Category = SubmixEffectPreset)
+	FConvolutionReverbSettings Settings;
+};
+
diff --git a/Engine/Plugins/Runtime/Synthesis/Source/Synthesis/Public/ImpulseResponseAsset.h b/Engine/Plugins/Runtime/Synthesis/Source/Synthesis/Public/ImpulseResponseAsset.h
new file mode 100644
--- /dev/null
+++ b/Engine/Plugins/Runtime/Synthesis/Source/Synthesis/Public/ImpulseResponseAsset.h
@@ -0,0 +1,125 @@
+// Copyright Epic Games, Inc. All Rights Reserved.
+
+#pragma once
+
+#include "CoreMinimal.h"
+#include "DSP/BufferVectorOperations.h"
+#include "UObject/ObjectMacros.h"
+#include "UObject/Object.h"
+#include "UObject/Class.h"
+#include "Engine/EngineTypes.h"
+#include "Engine/DataTable.h"
+#include "AudioEffect.h"
+#include "ImpulseResponseAsset.generated.h"
+
+
+// ========================================================================
+// ImpulseResponesConverter:
+//
+// Asynchronously prepare IR data for Audio Render Thread
+// It owns no data.
+// It is owned by UImpulseResponse asset, so it can safely assume it
+//	hasn't been GC'd.
+// ========================================================================
+class UImpulseResponse;
+
+class SYNTHESIS_API ImpulseResponesConverter : public FNonAbandonableTask
+{
+public:
+	ImpulseResponesConverter(UImpulseResponse* InIrAsset, int32 InSampleRate)
+		: IrAssetPtr(InIrAsset)
+		, TargetSampleRate(InSampleRate)
+	{
+	}
+
+private:
+	void DoWork();
+
+	// Pointer to the UImpulseResponse asset
+	UImpulseResponse* IrAssetPtr;
+
+	// needed to perform SRC in this task
+	int32 TargetSampleRate;
+
+	FORCEINLINE TStatId GetStatId() const
+	{
+		RETURN_QUICK_DECLARE_CYCLE_STAT(ImpulseResponesConverter, STATGROUP_ThreadPoolAsyncTasks);
+	}
+
+	friend class FAsyncTask<ImpulseResponesConverter>;
+};
+
+
+// ========================================================================
+// UImpulseResponse
+// UAsset used to represent Imported Impulse Responses
+// ========================================================================
+UCLASS()
+class SYNTHESIS_API UImpulseResponse : public UObject
+{
+	GENERATED_BODY()
+
+public:
+
+	UImpulseResponse();
+
+	~UImpulseResponse();
+
+	void SubscribeFlagToIrChanges(FThreadSafeBool* ExternalThreadSafeBool);
+
+	void UnsubscribeFlagFromIrChanges(FThreadSafeBool* ExternalThreadSafeBool);
+
+	const Audio::AlignedFloatBuffer& GetDeinterleavedIRData() const { return DeinterleavedIR; }
+
+	// If the data needs de-interleave/SRC, it will do so in an async task
+	// The async task will call PrepareIRsForConvolutionAlgorithmInternal() on another thread...
+	// and call SetSubscriberFlags(true) when it is done.
+	// If the data is already ready, SetSubscriberFlags(true) will be called synchronously
+	void PrepareIRsForConvolutionAlgorithm(const int32 InSampleRate);
+
+	// returns true if the cache for the Convolution Algorithm object is already warm (no async task needed)
+	// reset task pointer if it is pointing to an old task
+	bool IsDataPrepared(const int32 InSampleRate);
+
+	// functions that get info for convolution algorithm based in IR format
+	int32 GetNumChannels() { return NumChannels; }
+
+	int32 GetNumFrames() { return NumFrames; }
+
+protected:
+	void PrepareIRsForConvolutionAlgorithmInternal(const int32 InSampleRate);
+
+	void SetSubscriberFlags(bool bInDataIsReady);
+
+	void PerformAssetIrSRC(int32 InSampleRate);
+
+	// flip these to true when new IR data is ready
+	TArray<FThreadSafeBool*> SubscriberFlags;
+
+	FCriticalSection SubscriberFlagsCritSec;
+
+	TUniquePtr<FAsyncTask<ImpulseResponesConverter>> TaskPtr;
+
+	// Buffer that gets copied by the ConvolutionAlgorithmWrapper
+	Audio::AlignedFloatBuffer DeinterleavedIR;
+
+	// Sample rate of the DeinteleavedIR buffer
+	int32 CurrentSampleRate{ 0 };
+
+	// data about the IR asset (filled out by the factory)
+	UPROPERTY()
+	TArray<float> IRData;
+
+	UPROPERTY()
+	int32 NumChannels;
+
+	UPROPERTY()
+	int32 NumFrames;
+
+	// sample rate of the SERIALIZED data, not unnecessarily the SR of current project
+	UPROPERTY()
+	int32 AssetSampleRate;
+
+	friend class UImpulseResponseFactory; // allow factory to fill the internal buffer
+	friend class ImpulseResponesConverter; // allow async worker to raise flags upon completion
+};
\ No newline at end of file
diff --git a/Engine/Plugins/Runtime/Synthesis/Source/SynthesisEditor/Classes/ImpulseResponse.h b/Engine/Plugins/Runtime/Synthesis/Source/SynthesisEditor/Classes/ImpulseResponse.h
new file mode 100644
--- /dev/null
+++ b/Engine/Plugins/Runtime/Synthesis/Source/SynthesisEditor/Classes/ImpulseResponse.h
@@ -0,0 +1,37 @@
+// Copyright Epic Games, Inc. All Rights Reserved.
+
+#pragma once
+
+#include "CoreMinimal.h"
+#include "Stats/Stats.h"
+#include "AssetToolsModule.h"
+#include "Factories/Factory.h"
+#include "AssetTypeActions_Base.h"
+#include "Sound/SampleBufferIO.h"
+#include "ConvolutionReverbComponent.h"
+#include "ImpulseResponse.generated.h"
+
+
+class FAssetTypeActions_ImpulseResponse : public FAssetTypeActions_Base
+{
+public:
+	virtual FText GetName() const override { return NSLOCTEXT("AssetTypeActions", "AssetTypeActions_ImpulseResponse", "Impulse Response"); }
+	virtual FColor GetTypeColor() const override { return FColor(0, 255, 0); }
+	virtual UClass* GetSupportedClass() const override;
+	virtual const TArray<FText>& GetSubMenus() const override;
+	virtual uint32 GetCategories() override { return EAssetTypeCategories::Sounds; }
+};
+
+UCLASS(MinimalAPI, hidecategories = Object)
+class UImpulseResponseFactory : public UFactory
+{
+	GENERATED_UCLASS_BODY()
+
+	//~ Begin UFactory Interface
+	virtual UObject* FactoryCreateNew(UClass* Class, UObject* InParent, FName Name, EObjectFlags Flags, UObject* Context, FFeedbackContext* Warn) override;
+	//~ Begin UFactory Interface
+	TWeakObjectPtr<USoundWave> StagedSoundWave;
+
+private:
+	Audio::FSoundWavePCMLoader Loader;
+};
diff --git a/Engine/Plugins/Runtime/Synthesis/Source/SynthesisEditor/Private/ImpulseResponse.cpp b/Engine/Plugins/Runtime/Synthesis/Source/SynthesisEditor/Private/ImpulseResponse.cpp
new file mode 100644
--- /dev/null
+++ b/Engine/Plugins/Runtime/Synthesis/Source/SynthesisEditor/Private/ImpulseResponse.cpp
@@ -0,0 +1,76 @@
+// Copyright Epic Games, Inc. All Rights Reserved.
+#include "ImpulseResponse.h"
+#include "ConvolutionReverbComponent.h"
+#include "UObject/ObjectMacros.h"
+#include "UObject/Object.h"
+
+#define LOCTEXT_NAMESPACE "AssetTypeActions"
+
+
+UClass* FAssetTypeActions_ImpulseResponse::GetSupportedClass() const
+{
+	return UImpulseResponse::StaticClass();
+}
+
+const TArray<FText>& FAssetTypeActions_ImpulseResponse::GetSubMenus() const
+{
+	static const TArray<FText> SubMenus
+	{
+		FText(LOCTEXT("AssetConvolutionReverbSubmenu", "Convolution Reverb"))
+	};
+
+	return SubMenus;
+}
+
+UImpulseResponseFactory::UImpulseResponseFactory(const FObjectInitializer& ObjectInitializer)
+	: Super(ObjectInitializer)
+{
+	SupportedClass = UImpulseResponse::StaticClass();
+
+	bCreateNew = true;
+	bEditorImport = false;
+	bEditAfterNew = true;
+}
+
+UObject* UImpulseResponseFactory::FactoryCreateNew(UClass* Class, UObject* InParent, FName InName, EObjectFlags Flags, UObject* Context, FFeedbackContext* Warn)
+{
+	UImpulseResponse* NewAsset = NewObject<UImpulseResponse>(InParent, InName, Flags);
+
+	if (StagedSoundWave.IsValid())
+	{
+		USoundWave* Wave = StagedSoundWave.Get();
+
+		Loader.LoadSoundWave(Wave, [&](const USoundWave* SoundWave, const Audio::FSampleBuffer& LoadedSampleBuffer)
+		{
+			NewAsset->NumChannels = LoadedSampleBuffer.GetNumChannels();
+			NewAsset->NumFrames = LoadedSampleBuffer.GetNumFrames();
+			NewAsset->AssetSampleRate = LoadedSampleBuffer.GetSampleRate();
+
+			// Convert to float and de-interleave (TODO: SRC?)
+			const int32 NumChannels = NewAsset->NumChannels;
+			const int32 NumFrames = NewAsset->NumFrames;
+			const int32 NumSamples = NumChannels * NumFrames;
+
+			NewAsset->IRData.Reset();
+			NewAsset->IRData.AddUninitialized(NumSamples);
+
+			const int16* RESTRICT InputBuffer = reinterpret_cast<const int16*>(LoadedSampleBuffer.GetData());
+			float* RESTRICT OutputBuffer = NewAsset->IRData.GetData();
+
+			for (int32 i = 0; i < NumSamples; ++i)
+			{
+				float CurrSample = static_cast<float>(InputBuffer[i]) / 32768.0f;
+				const int32 CurrFrame = i / NumChannels;
+				const int32 CurrChannel = i % NumChannels;
+				const int32 CurrOutputIndex = CurrChannel * NumFrames + CurrFrame;
+				OutputBuffer[CurrOutputIndex] = CurrSample;
+			}
+		});
+
+		StagedSoundWave.Reset();
+	}
+
+	return NewAsset;
+}
+
+#undef LOCTEXT_NAMESPACE
\ No newline at end of file
diff --git a/Engine/Plugins/Runtime/Synthesis/Source/SynthesisEditor/Private/SoundWaveAssetActionExtenderConvolutionReverb.cpp b/Engine/Plugins/Runtime/Synthesis/Source/SynthesisEditor/Private/SoundWaveAssetActionExtenderConvolutionReverb.cpp
new file mode 100644
--- /dev/null
+++ b/Engine/Plugins/Runtime/Synthesis/Source/SynthesisEditor/Private/SoundWaveAssetActionExtenderConvolutionReverb.cpp
@@ -0,0 +1,71 @@
+// Copyright Epic Games, Inc. All Rights Reserved.
+
+#include "SoundWaveAssetActionExtenderConvolutionReverb.h"
+#include "ToolMenus.h"
+#include "AssetTypeActions_Base.h"
+#include "IContentBrowserSingleton.h"
+#include "ContentBrowserModule.h"
+#include "ContentBrowserMenuContexts.h"
+#include "ObjectEditorUtils.h"
+#include "EditorStyleSet.h"
+#include "Sound/SoundWave.h"
+#include "ImpulseResponse.h"
+#include "ConvolutionReverbComponent.h"
+
+#define LOCTEXT_NAMESPACE "AssetTypeActions"
+
+void FSoundWaveAssetActionExtenderConvolutionReverb::RegisterMenus()
+{
+	if (!UToolMenus::IsToolMenuUIEnabled())
+	{
+		return;
+	}
+
+	FToolMenuOwnerScoped MenuOwner("ConvolutionReverb");
+	UToolMenu* Menu = UToolMenus::Get()->ExtendMenu("ContentBrowser.AssetContextMenu.SoundWave");
+	FToolMenuSection& Section = Menu->FindOrAddSection("GetAssetActions");
+
+	Section.AddDynamicEntry("SoundWaveAssetConversion", FNewToolMenuSectionDelegate::CreateLambda([](FToolMenuSection& InSection)
+	{
+		const TAttribute<FText> Label = LOCTEXT("SoundWave_CreateImpulseResponse", "Create Impulse Response");
+		const TAttribute<FText> ToolTip = LOCTEXT("SoundWave_CreateImpulseResponseTooltip", "Creates an impulse response asset using the selected sound wave.");
+		const FSlateIcon Icon = FSlateIcon(FEditorStyle::GetStyleSetName(), "ClassIcon.ImpulseResponse");
+		const FToolMenuExecuteAction UIAction = FToolMenuExecuteAction::CreateStatic(&FSoundWaveAssetActionExtenderConvolutionReverb::ExecuteCreateImpulseResponse);
+
+		InSection.AddMenuEntry("SoundWave_CreateImpulseResponse", Label, ToolTip, Icon, UIAction);
+	}));
+}
+
+void FSoundWaveAssetActionExtenderConvolutionReverb::ExecuteCreateImpulseResponse(const FToolMenuContext& MenuContext)
+{
+	UContentBrowserAssetContextMenuContext* Context = MenuContext.FindContext<UContentBrowserAssetContextMenuContext>();
+	if (!Context || Context->SelectedObjects.Num() == 0)
+	{
+		return;
+	}
+
+	const FString DefaultSuffix = TEXT("_IR");
+	FAssetToolsModule& AssetToolsModule = FModuleManager::Get().LoadModuleChecked<FAssetToolsModule>("AssetTools");
+
+	// Create the factory used to generate the asset
+	UImpulseResponseFactory* Factory = NewObject<UImpulseResponseFactory>();
+	
+	// only converts 0th selected object for now (see return statement)
+	for (const TWeakObjectPtr<UObject>& Object : Context->SelectedObjects)
+	{
+		// stage the soundwave on the factory to be used during asset creation
+		USoundWave* Wave = Cast<USoundWave>(Object);
+		check(Wave);
+		Factory->StagedSoundWave = Wave; // WeakPtr gets reset by the Factory after it is consumed
+
+		// Determine an appropriate name
+		FString Name;
+		FString PackagePath;
+		AssetToolsModule.Get().CreateUniqueAssetName(Wave->GetOutermost()->GetName(), DefaultSuffix, PackagePath, Name);
+
+		// create new asset
+		AssetToolsModule.Get().CreateAsset(Name, FPackageName::GetLongPackagePath(PackagePath), UImpulseResponse::StaticClass(), Factory);
+	}
+}
+
+#undef LOCTEXT_NAMESPACE
\ No newline at end of file
diff --git a/Engine/Plugins/Runtime/Synthesis/Source/SynthesisEditor/Private/SynthesisEditorModule.cpp b/Engine/Plugins/Runtime/Synthesis/Source/SynthesisEditor/Private/SynthesisEditorModule.cpp
--- a/Engine/Plugins/Runtime/Synthesis/Source/SynthesisEditor/Private/SynthesisEditorModule.cpp
+++ b/Engine/Plugins/Runtime/Synthesis/Source/SynthesisEditor/Private/SynthesisEditorModule.cpp
@@ -12,6 +12,9 @@
 #include "AudioEditorModule.h"
 #include "EpicSynth1PresetBank.h"
 #include "MonoWaveTablePresetBank.h"
+#include "ImpulseResponse.h"
+#include "SoundWaveAssetActionExtenderConvolutionReverb.h"
+#include "ToolMenus.h"
 
 IMPLEMENT_MODULE(FSynthesisEditorModule, SynthesisEditor)
 
@@ -21,14 +24,20 @@ void FSynthesisEditorModule::StartupModule()
 
 	AssetTools.RegisterAssetTypeActions(MakeShareable(new FAssetTypeActions_ModularSynthPresetBank));
 	AssetTools.RegisterAssetTypeActions(MakeShareable(new FAssetTypeActions_MonoWaveTableSynthPreset));
+	AssetTools.RegisterAssetTypeActions(MakeShareable(new FAssetTypeActions_ImpulseResponse));
 
 	// Now that we've loaded this module, we need to register our effect preset actions
 	IAudioEditorModule* AudioEditorModule = &FModuleManager::LoadModuleChecked<IAudioEditorModule>("AudioEditor");
 	AudioEditorModule->RegisterEffectPresetAssetActions();
+
+	UToolMenus::RegisterStartupCallback(FSimpleMulticastDelegate::FDelegate::CreateRaw(this, &FSynthesisEditorModule::RegisterMenus));
 }
 
 void FSynthesisEditorModule::ShutdownModule()
 {
 }
 
-
+void FSynthesisEditorModule::RegisterMenus()
+{
+	FSoundWaveAssetActionExtenderConvolutionReverb::RegisterMenus();
+}
\ No newline at end of file
diff --git a/Engine/Plugins/Runtime/Synthesis/Source/SynthesisEditor/Public/SoundWaveAssetActionExtenderConvolutionReverb.h b/Engine/Plugins/Runtime/Synthesis/Source/SynthesisEditor/Public/SoundWaveAssetActionExtenderConvolutionReverb.h
new file mode 100644
--- /dev/null
+++ b/Engine/Plugins/Runtime/Synthesis/Source/SynthesisEditor/Public/SoundWaveAssetActionExtenderConvolutionReverb.h
@@ -0,0 +1,12 @@
+// Copyright Epic Games, Inc. All Rights Reserved.
+
+#pragma once
+
+class FSoundWaveAssetActionExtenderConvolutionReverb
+{
+public:
+	static void RegisterMenus();
+	static void GetExtendedActions(const struct FToolMenuContext& MenuContext);
+	static void ExecuteCreateImpulseResponse(const struct FToolMenuContext& MenuContext);
+};
+
diff --git a/Engine/Plugins/Runtime/Synthesis/Source/SynthesisEditor/Public/SynthesisEditorModule.h b/Engine/Plugins/Runtime/Synthesis/Source/SynthesisEditor/Public/SynthesisEditorModule.h
--- a/Engine/Plugins/Runtime/Synthesis/Source/SynthesisEditor/Public/SynthesisEditorModule.h
+++ b/Engine/Plugins/Runtime/Synthesis/Source/SynthesisEditor/Public/SynthesisEditorModule.h
@@ -12,4 +12,7 @@ public:
 	// IModuleInterface interface
 	virtual void StartupModule() override;
 	virtual void ShutdownModule() override;
+
+private:
+	void RegisterMenus();
 };
diff --git a/Engine/Plugins/Runtime/Synthesis/Source/SynthesisEditor/SynthesisEditor.Build.cs b/Engine/Plugins/Runtime/Synthesis/Source/SynthesisEditor/SynthesisEditor.Build.cs
--- a/Engine/Plugins/Runtime/Synthesis/Source/SynthesisEditor/SynthesisEditor.Build.cs
+++ b/Engine/Plugins/Runtime/Synthesis/Source/SynthesisEditor/SynthesisEditor.Build.cs
@@ -6,6 +6,7 @@ namespace UnrealBuildTool.Rules
 	{
         public SynthesisEditor(ReadOnlyTargetRules Target) : base(Target)
 		{
+
             PublicDependencyModuleNames.AddRange(
 				new string[] {
                     "Core",
@@ -15,7 +16,12 @@ namespace UnrealBuildTool.Rules
 					"AudioEditor",
                     "Synthesis",
 					"AudioMixer",
-				}
+                    "ToolMenus",
+                    "EditorStyle",
+                    "Slate",
+                    "SlateCore",
+                    "ContentBrowser",
+                }
 			);
 
             PrivateIncludePathModuleNames.AddRange(
diff --git a/Engine/Source/Runtime/SignalProcessing/Public/DSP/ConvolutionAlgorithm.h b/Engine/Source/Runtime/SignalProcessing/Public/DSP/ConvolutionAlgorithm.h
--- a/Engine/Source/Runtime/SignalProcessing/Public/DSP/ConvolutionAlgorithm.h
+++ b/Engine/Source/Runtime/SignalProcessing/Public/DSP/ConvolutionAlgorithm.h
@@ -30,6 +30,21 @@ namespace Audio
 
 		/** Maximum number of samples in an impulse responses. */
 		int32 MaxNumImpulseResponseSamples = 0;
+
+		bool operator==(const FConvolutionSettings& Other)
+		{
+			return bEnableHardwareAcceleration == Other.bEnableHardwareAcceleration
+				&& BlockNumSamples == Other.BlockNumSamples
+				&& NumInputChannels == Other.NumInputChannels
+				&& NumOutputChannels == Other.NumOutputChannels
+				&& NumImpulseResponses == Other.NumImpulseResponses
+				&& MaxNumImpulseResponseSamples == Other.MaxNumImpulseResponseSamples;
+		}
+
+		bool operator!=(const FConvolutionSettings& Other)
+		{
+			return !(*this == Other);
+		}
 	};
 
 
