Fixed SRC bugs in Impulse Response import

#jra none

diff --git a/Engine/Plugins/Runtime/Synthesis/Source/Synthesis/Private/ImpulseResponseAsset.cpp b/Engine/Plugins/Runtime/Synthesis/Source/Synthesis/Private/ImpulseResponseAsset.cpp
--- a/Engine/Plugins/Runtime/Synthesis/Source/Synthesis/Private/ImpulseResponseAsset.cpp
+++ b/Engine/Plugins/Runtime/Synthesis/Source/Synthesis/Private/ImpulseResponseAsset.cpp
@@ -139,8 +139,9 @@ void UImpulseResponse::PerformAssetIrSRC(int32 InSampleRate)
 {
 	const float StepSize = AssetSampleRate / static_cast<float>(InSampleRate);
 	const int32 AssetBufferSize = IRData.Num();
-	const int32 FinalBufferSize = static_cast<int32>(StepSize * AssetBufferSize);
-	const int32 FinalNumFrames = static_cast<int32>(StepSize * NumFrames);
+	const int32 AssetBufferNumFrames = AssetBufferSize / NumChannels;
+	const int32 FinalBufferSize = static_cast<int32>(AssetBufferSize / StepSize);
+	const int32 FinalNumFrames = static_cast<int32>(NumFrames / StepSize);
 	float SampleA, SampleB;
 
 	DeinterleavedIR.Reset();
@@ -153,7 +154,7 @@ void UImpulseResponse::PerformAssetIrSRC(int32 InSampleRate)
 		return;
 	}
 
-	// SRC: (need to test this with an IR that has a different SR than the FAudioDevice)
+	// SRC
 	// for each channel...
 	for (int32 Chan = 0; Chan < NumChannels; ++Chan)
 	{
@@ -171,7 +172,7 @@ void UImpulseResponse::PerformAssetIrSRC(int32 InSampleRate)
 			// check for interpolation between last and first frames
 			SampleA = IRData[InputChanOffset + WholeThisIndex];
 
-			if (WholeNextIndex != AssetBufferSize)
+			if (WholeNextIndex != AssetBufferNumFrames)
 			{
 				SampleB = IRData[InputChanOffset + WholeNextIndex];
 			}
@@ -183,6 +184,11 @@ void UImpulseResponse::PerformAssetIrSRC(int32 InSampleRate)
 			DeinterleavedIR[OutputChanOffset + i] = FMath::Lerp(SampleA, SampleB, Alpha);
 
 			CurrIndex += StepSize;
+
+			if ((FMath::FloorToInt(CurrIndex) + InputChanOffset) == AssetBufferSize)
+			{
+				break;
+			}
 		}
 	}
 }
\ No newline at end of file
