Enable synthesis plugin by default

diff --git a/Engine/Plugins/Runtime/Synthesis/Synthesis.uplugin b/Engine/Plugins/Runtime/Synthesis/Synthesis.uplugin
--- a/Engine/Plugins/Runtime/Synthesis/Synthesis.uplugin
+++ b/Engine/Plugins/Runtime/Synthesis/Synthesis.uplugin
@@ -1,7 +1,7 @@
 {
 	"FileVersion" : 3,
 	"Version" : 1,
-	"VersionName" : "1.0",
+	"VersionName" : "1.1",
 	"FriendlyName" : "Synthesis and DSP Effects",
 	"Description" : "A variety of realtime synthesizers and DSP source and submix effects.",
 	"Category" : "Audio",
@@ -10,9 +10,9 @@
 	"DocsURL" : "",
 	"MarketplaceURL" : "",
 	"SupportURL" : "",
-	"EnabledByDefault" : false,
+	"EnabledByDefault" : true,
 	"CanContainContent" : false,
-	"IsBetaVersion" : true,
+	"IsBetaVersion" : false,
 	"Installed" : false,
 	"Modules" :
 	[
