- Move Effect Preset Instances array to use weak pointers to allow for safe bookkeeping and avoid crash on shutdown with dangling raw pointer after client mixer code deletes out from under effect system.

Note: Cleaned up resource acquisition and initialization to be more clear as to who owns sound effect instances using factory-style Create call on FSoundEffectSubmix & FSoundEffectSource.

#rnx

diff --git a/Engine/Plugins/Runtime/Oculus/OculusAudio/Source/OculusAudio/Private/OculusAudioReverb.cpp b/Engine/Plugins/Runtime/Oculus/OculusAudio/Source/OculusAudio/Private/OculusAudioReverb.cpp
--- a/Engine/Plugins/Runtime/Oculus/OculusAudio/Source/OculusAudio/Private/OculusAudioReverb.cpp
+++ b/Engine/Plugins/Runtime/Oculus/OculusAudio/Source/OculusAudio/Private/OculusAudioReverb.cpp
@@ -68,7 +68,14 @@ FSoundEffectSubmixPtr OculusAudioReverb::GetEffectSubmix()
 {
 	if (!SubmixEffect.IsValid())
 	{
-		SubmixEffect = MakeShareable(static_cast<FSoundEffectSubmix*>(new FSubmixEffectOculusReverbPlugin()));
+		if (!ReverbPreset)
+		{
+			ReverbPreset = NewObject<USubmixEffectOculusReverbPreset>();
+			ReverbPreset->AddToRoot();
+		}
+
+		FSoundEffectSubmixInitData InitData;
+		SubmixEffect = FSoundEffectSubmix::Create(InitData, *ReverbPreset);
 		SubmixEffect->SetEnabled(true);
 	}
 
diff --git a/Engine/Plugins/Runtime/Oculus/OculusAudio/Source/OculusAudio/Private/OculusAudioReverb.h b/Engine/Plugins/Runtime/Oculus/OculusAudio/Source/OculusAudio/Private/OculusAudioReverb.h
--- a/Engine/Plugins/Runtime/Oculus/OculusAudio/Source/OculusAudio/Private/OculusAudioReverb.h
+++ b/Engine/Plugins/Runtime/Oculus/OculusAudio/Source/OculusAudio/Private/OculusAudioReverb.h
@@ -4,15 +4,23 @@
 
 #include "IAudioExtensionPlugin.h"
 #include "Sound/SoundEffectSubmix.h"
+#include "ObjectMacros.h"
 #include "OVR_Audio.h"
 #include "Templates/UniquePtr.h"
+#include "Sound/SoundEffectBase.h"
 
 
-class FSubmixEffectOculusReverbPlugin : public FSoundEffectSubmix
+UCLASS()
+class USubmixEffectOculusReverbPreset : public USoundEffectSubmixPreset
 {
+	GENERATED_BODY()
+
 public:
-	FSubmixEffectOculusReverbPlugin();
+	EFFECT_PRESET_METHODS(SubmixEffectOculusReverbPreset)
+};
 
+class FSubmixEffectOculusReverbPlugin : public FSoundEffectSubmix
+{
 	void ClearContext();
 
 	virtual void Init(const FSoundEffectSubmixInitData& InInitData) override;
@@ -26,6 +34,10 @@ public:
 	{
 		return; // PAS
 	}
+
+protected:
+	FSubmixEffectOculusReverbPlugin();
+
 private:
 	ovrAudioContext Context;
 	FCriticalSection ContextLock;
@@ -41,6 +53,7 @@ class OculusAudioReverb : public IAudioReverb
 public:
 	OculusAudioReverb()
 		: Context(nullptr)
+		, ReverbPreset(nullptr)
 	{
 		// empty
 	}
@@ -67,5 +80,6 @@ public:
 	}
 private:
 	ovrAudioContext* Context;
-	FSoundEffectSubmixPtr SubmixEffect;
+	TSoundEffectSubmixPtr SubmixEffect;
+	USubmixEffectOculusReverbPreset* ReverbPreset;
 };
\ No newline at end of file
diff --git a/Engine/Plugins/Runtime/ResonanceAudio/Source/ResonanceAudio/Private/ResonanceAudioReverb.cpp b/Engine/Plugins/Runtime/ResonanceAudio/Source/ResonanceAudio/Private/ResonanceAudioReverb.cpp
--- a/Engine/Plugins/Runtime/ResonanceAudio/Source/ResonanceAudio/Private/ResonanceAudioReverb.cpp
+++ b/Engine/Plugins/Runtime/ResonanceAudio/Source/ResonanceAudio/Private/ResonanceAudioReverb.cpp
@@ -139,12 +139,9 @@ namespace ResonanceAudio
 			GlobalReverbPluginPreset->AddToRoot();
 			ReverbPluginPreset = GlobalReverbPluginPreset;
 
-			FResonanceAudioReverbPlugin* ReverbSubmixEffect = static_cast<FResonanceAudioReverbPlugin*>(ReverbPluginPreset->CreateNewEffect());
-			ReverbSubmixEffect->SetPreset(GlobalReverbPluginPreset);
-			ReverbSubmixEffect->SetResonanceAudioReverbPlugin(this);
-
-			ReverbSubmixEffect->SetEnabled(true);
-			SubmixEffect = MakeShareable(static_cast<FSoundEffectSubmix*>(ReverbSubmixEffect));
+			SubmixEffect = USoundEffectPreset::CreateInstance<FSoundEffectSubmixInitData, FSoundEffectSubmix>(FSoundEffectSubmixInitData(), *ReverbPluginPreset);
+			StaticCastSharedPtr<FResonanceAudioReverbPlugin, FSoundEffectSubmix, ESPMode::ThreadSafe>(SubmixEffect)->SetResonanceAudioReverbPlugin(this);
+			SubmixEffect->SetEnabled(true);
 		}
 		else
 		{
diff --git a/Engine/Plugins/Runtime/Steam/SteamAudio/Source/SteamAudio/Private/PhononReverb.cpp b/Engine/Plugins/Runtime/Steam/SteamAudio/Source/SteamAudio/Private/PhononReverb.cpp
--- a/Engine/Plugins/Runtime/Steam/SteamAudio/Source/SteamAudio/Private/PhononReverb.cpp
+++ b/Engine/Plugins/Runtime/Steam/SteamAudio/Source/SteamAudio/Private/PhononReverb.cpp
@@ -410,13 +410,9 @@ namespace SteamAudio
 				ReverbPluginPreset = NewObject<USubmixEffectReverbPluginPreset>(Submix, TEXT("Reverb Plugin Effect Preset"));
 			}
 
-			FSubmixEffectReverbPlugin* ReverbSubmixEffect = static_cast<FSubmixEffectReverbPlugin*>(ReverbPluginPreset->CreateNewEffect());
-			ReverbSubmixEffect->SetPreset(ReverbPluginPreset);
-			ReverbSubmixEffect->SetPhononReverbPlugin(this);
-
-			ReverbSubmixEffect->SetEnabled(true);
-
-			SubmixEffect = MakeShareable(static_cast<FSoundEffectSubmix*>(ReverbSubmixEffect));
+			SubmixEffect = FSoundEffectSubmix::Create(FSoundEffectSubmixInitData(), *ReverbPluginPreset);
+			StaticCastSharedPtr<FSubmixEffectReverbPlugin, FSoundEffectSubmix, ESPMode::ThreadSafe>(SubmixEffect)->SetPhononReverbPlugin(this);
+			SubmixEffect->SetEnabled(true);
 		}
 
 		return SubmixEffect;
diff --git a/Engine/Plugins/Runtime/Synthesis/Source/Synthesis/Classes/SubmixEffects/SubmixEffectFilter.h b/Engine/Plugins/Runtime/Synthesis/Source/Synthesis/Classes/SubmixEffects/SubmixEffectFilter.h
--- a/Engine/Plugins/Runtime/Synthesis/Source/Synthesis/Classes/SubmixEffects/SubmixEffectFilter.h
+++ b/Engine/Plugins/Runtime/Synthesis/Source/Synthesis/Classes/SubmixEffects/SubmixEffectFilter.h
@@ -66,7 +66,7 @@ class SYNTHESIS_API FSubmixEffectFilter : public FSoundEffectSubmix
 {
 public:
 	FSubmixEffectFilter();
-	~FSubmixEffectFilter();
+	virtual ~FSubmixEffectFilter() = default;
 
 	//~ Begin FSoundEffectSubmix
 	virtual void Init(const FSoundEffectSubmixInitData& InData) override;
diff --git a/Engine/Plugins/Runtime/Synthesis/Source/Synthesis/Classes/SubmixEffects/SubmixEffectTapDelay.h b/Engine/Plugins/Runtime/Synthesis/Source/Synthesis/Classes/SubmixEffects/SubmixEffectTapDelay.h
--- a/Engine/Plugins/Runtime/Synthesis/Source/Synthesis/Classes/SubmixEffects/SubmixEffectTapDelay.h
+++ b/Engine/Plugins/Runtime/Synthesis/Source/Synthesis/Classes/SubmixEffects/SubmixEffectTapDelay.h
@@ -110,7 +110,8 @@ class SYNTHESIS_API FSubmixEffectTapDelay : public FSoundEffectSubmix
 {
 public:
 	FSubmixEffectTapDelay();
-	~FSubmixEffectTapDelay();
+
+	virtual ~FSubmixEffectTapDelay() = default;
 
 	//~ Begin FSoundEffectSubmix
 	virtual void Init(const FSoundEffectSubmixInitData& InData) override;
diff --git a/Engine/Plugins/Runtime/Synthesis/Source/Synthesis/Private/SubmixEffectDelay.cpp b/Engine/Plugins/Runtime/Synthesis/Source/Synthesis/Private/SubmixEffectDelay.cpp
--- a/Engine/Plugins/Runtime/Synthesis/Source/Synthesis/Private/SubmixEffectDelay.cpp
+++ b/Engine/Plugins/Runtime/Synthesis/Source/Synthesis/Private/SubmixEffectDelay.cpp
@@ -1,5 +1,7 @@
 // Copyright Epic Games, Inc. All Rights Reserved.
+
 #include "SubmixEffects/SubmixEffectDelay.h"
+#include "Sound/SoundEffectPreset.h"
 
 const float FSubmixEffectDelay::MinLengthDelaySec = 0.4f;
 
@@ -130,29 +132,23 @@ void FSubmixEffectDelay::OnNumChannelsChanged(int32 NumChannels)
 void USubmixEffectDelayPreset::SetInterpolationTime(float Time)
 {
 	DynamicSettings.InterpolationTime = Time;
+
 	// Dispatch to all effect instances:
-	for (FSoundEffectBase* EffectBaseInstance : Instances)
+	EffectCommand<FSubmixEffectDelay>([Time](FSubmixEffectDelay& TapDelay)
 	{
-		FSubmixEffectDelay* TapDelay = (FSubmixEffectDelay*)EffectBaseInstance;
-		EffectBaseInstance->EffectCommand([TapDelay, Time]()
-		{
-			TapDelay->SetInterpolationTime(Time);
-		});
-	}
+		TapDelay.SetInterpolationTime(Time);
+	});
 }
 
 void USubmixEffectDelayPreset::SetDelay(float Length)
 {
 	DynamicSettings.DelayLength = Length;
+
 	// Dispatch to all effect instances:
-	for (FSoundEffectBase* EffectBaseInstance : Instances)
+	EffectCommand<FSubmixEffectDelay>([Length](FSubmixEffectDelay& TapDelay)
 	{
-		FSubmixEffectDelay* Delay = (FSubmixEffectDelay*)EffectBaseInstance;
-		EffectBaseInstance->EffectCommand([Delay, Length]()
-		{
-			Delay->SetDelayLineLength(Length);
-		});
-	}
+		TapDelay.SetDelayLineLength(Length);
+	});
 }
 
 void USubmixEffectDelayPreset::OnInit()
diff --git a/Engine/Plugins/Runtime/Synthesis/Source/Synthesis/Private/SubmixEffectFilter.cpp b/Engine/Plugins/Runtime/Synthesis/Source/Synthesis/Private/SubmixEffectFilter.cpp
--- a/Engine/Plugins/Runtime/Synthesis/Source/Synthesis/Private/SubmixEffectFilter.cpp
+++ b/Engine/Plugins/Runtime/Synthesis/Source/Synthesis/Private/SubmixEffectFilter.cpp
@@ -3,7 +3,6 @@
 #include "SubmixEffects/SubmixEffectFilter.h"
 #include "AudioMixer.h"
 
-PRAGMA_DISABLE_OPTIMIZATION
 
 FSubmixEffectFilter::FSubmixEffectFilter()
 	: SampleRate(0.0f)
@@ -18,10 +17,6 @@ FSubmixEffectFilter::FSubmixEffectFilter()
 {
 }
 
-FSubmixEffectFilter::~FSubmixEffectFilter()
-{
-}
-
 void FSubmixEffectFilter::Init(const FSoundEffectSubmixInitData& InData)
 {
 	SampleRate = InData.SampleRate;
@@ -167,75 +162,48 @@ void USubmixEffectFilterPreset::SetSettings(const FSubmixEffectFilterSettings& I
 
 void USubmixEffectFilterPreset::SetFilterType(ESubmixFilterType InType)
 {
-	for (FSoundEffectBase* EffectBaseInstance : Instances)
+	EffectCommand<FSubmixEffectFilter>([InType](FSubmixEffectFilter& FilterEffect)
 	{
-		FSubmixEffectFilter* FilterEffect = (FSubmixEffectFilter*)EffectBaseInstance;
-		EffectBaseInstance->EffectCommand([FilterEffect, InType]()
-		{
-			FilterEffect->SetFilterType(InType);
-		});
-	}
+		FilterEffect.SetFilterType(InType);
+	});
 }
 
 void USubmixEffectFilterPreset::SetFilterAlgorithm(ESubmixFilterAlgorithm InAlgorithm)
 {
-	for (FSoundEffectBase* EffectBaseInstance : Instances)
+	EffectCommand<FSubmixEffectFilter>([InAlgorithm](FSubmixEffectFilter& FilterEffect)
 	{
-		FSubmixEffectFilter* FilterEffect = (FSubmixEffectFilter*)EffectBaseInstance;
-		EffectBaseInstance->EffectCommand([FilterEffect, InAlgorithm]()
-		{
-			FilterEffect->SetFilterAlgorithm(InAlgorithm);
-		});
-	}
+		FilterEffect.SetFilterAlgorithm(InAlgorithm);
+	});
 }
 
 void USubmixEffectFilterPreset::SetFilterCutoffFrequency(float InFrequency)
 {
-	for (FSoundEffectBase* EffectBaseInstance : Instances)
+	EffectCommand<FSubmixEffectFilter>([InFrequency](FSubmixEffectFilter& FilterEffect)
 	{
-		FSubmixEffectFilter* FilterEffect = (FSubmixEffectFilter*)EffectBaseInstance;
-		EffectBaseInstance->EffectCommand([FilterEffect, InFrequency]()
-		{
-			FilterEffect->SetFilterCutoffFrequency(InFrequency);
-		});
-	}
+		FilterEffect.SetFilterCutoffFrequency(InFrequency);
+	});
 }
 
 void USubmixEffectFilterPreset::SetFilterCutoffFrequencyMod(float InFrequency)
 {
-	for (FSoundEffectBase* EffectBaseInstance : Instances)
+	EffectCommand<FSubmixEffectFilter>([InFrequency](FSubmixEffectFilter& FilterEffect)
 	{
-		FSubmixEffectFilter* FilterEffect = (FSubmixEffectFilter*)EffectBaseInstance;
-		EffectBaseInstance->EffectCommand([FilterEffect, InFrequency]()
-		{
-			FilterEffect->SetFilterCutoffFrequencyMod(InFrequency);
-		});
-	}
+		FilterEffect.SetFilterCutoffFrequencyMod(InFrequency);
+	});
 }
 
 void USubmixEffectFilterPreset::SetFilterQ(float InQ)
 {
-	for (FSoundEffectBase* EffectBaseInstance : Instances)
+	EffectCommand<FSubmixEffectFilter>([InQ](FSubmixEffectFilter& FilterEffect)
 	{
-		FSubmixEffectFilter* FilterEffect = (FSubmixEffectFilter*)EffectBaseInstance;
-		EffectBaseInstance->EffectCommand([FilterEffect, InQ]()
-		{
-			FilterEffect->SetFilterQ(InQ);
-		});
-	}
+		FilterEffect.SetFilterQ(InQ);
+	});
 }
 
 void USubmixEffectFilterPreset::SetFilterQMod(float InQ)
 {
-	for (FSoundEffectBase* EffectBaseInstance : Instances)
+	EffectCommand<FSubmixEffectFilter>([InQ](FSubmixEffectFilter& FilterEffect)
 	{
-		FSubmixEffectFilter* FilterEffect = (FSubmixEffectFilter*)EffectBaseInstance;
-		EffectBaseInstance->EffectCommand([FilterEffect, InQ]()
-		{
-			FilterEffect->SetFilterQMod(InQ);
-		});
-	}
-
+		FilterEffect.SetFilterQMod(InQ);
+	});
 }
-
-PRAGMA_ENABLE_OPTIMIZATION
\ No newline at end of file
diff --git a/Engine/Plugins/Runtime/Synthesis/Source/Synthesis/Private/SubmixEffectTapDelay.cpp b/Engine/Plugins/Runtime/Synthesis/Source/Synthesis/Private/SubmixEffectTapDelay.cpp
--- a/Engine/Plugins/Runtime/Synthesis/Source/Synthesis/Private/SubmixEffectTapDelay.cpp
+++ b/Engine/Plugins/Runtime/Synthesis/Source/Synthesis/Private/SubmixEffectTapDelay.cpp
@@ -3,69 +3,62 @@
 
 DEFINE_LOG_CATEGORY(LogTapDelay);
 
-static void GetMultichannelPan(const float& Input, const float& Angle, const float& Gain, const int32& NumChannels, float* DestinationFrame)
+namespace TapDelayUtils
 {
-	const float InputWithGain = Input * Gain;
-
-	if (NumChannels == 2)
+	void GetMultichannelPan(const float& Input, const float& Angle, const float& Gain, const int32& NumChannels, float* DestinationFrame)
 	{
-		float LeftGain;
-		float RightGain;
-		Audio::GetStereoPan(Angle / -90.0f, LeftGain, RightGain);
+		const float InputWithGain = Input * Gain;
 
-		DestinationFrame[0] += LeftGain * InputWithGain;
-		DestinationFrame[1] += RightGain * InputWithGain;
-	}
-	else
-	{
-		float QuadChannelMap[4] = { 0.0f, 0.0f, 0.0f, 0.0f };
+		if (NumChannels == 2)
+		{
+			float LeftGain;
+			float RightGain;
+			Audio::GetStereoPan(Angle / -90.0f, LeftGain, RightGain);
 
-		// We will spatialize using quad panning
-		// Note input angle results in 0.0 being behind you
-		// We offset the value to align the value so that 0.0 is "front left" i.e. index 0
-		// We use the modulo to make sure it wraps to be between 0.0 and 1.0
-		// We're also (for simplicity) assuming isotropic distribution of the quad speaker map.
-		const float NormalizedAngle = FMath::Fmod((Angle + 180.0f) / 360.0f + 0.625f, 1.0f);
+			DestinationFrame[0] += LeftGain * InputWithGain;
+			DestinationFrame[1] += RightGain * InputWithGain;
+		}
+		else
+		{
+			float QuadChannelMap[4] = { 0.0f, 0.0f, 0.0f, 0.0f };
 
-		const float ChannelFraction = 4 * NormalizedAngle;
-		const int32 Channel0 = FMath::FloorToInt(ChannelFraction);
-		const int32 Channel1 = (Channel0 + 1) % 4;
-		const float ChannelAlpha = (ChannelFraction - Channel0) * -2.0f + 1.0f;
+			// We will spatialize using quad panning
+			// Note input angle results in 0.0 being behind you
+			// We offset the value to align the value so that 0.0 is "front left" i.e. index 0
+			// We use the modulo to make sure it wraps to be between 0.0 and 1.0
+			// We're also (for simplicity) assuming isotropic distribution of the quad speaker map.
+			const float NormalizedAngle = FMath::Fmod((Angle + 180.0f) / 360.0f + 0.625f, 1.0f);
 
-		Audio::GetStereoPan(ChannelAlpha, QuadChannelMap[Channel0], QuadChannelMap[Channel1]);
+			const float ChannelFraction = 4 * NormalizedAngle;
+			const int32 Channel0 = FMath::FloorToInt(ChannelFraction);
+			const int32 Channel1 = (Channel0 + 1) % 4;
+			const float ChannelAlpha = (ChannelFraction - Channel0) * -2.0f + 1.0f;
 
-		// Now map to the specific channel outputs
-		if (NumChannels == 6 || NumChannels == 8)
-		{
-			// Specifically skipping LFE, Center channel, and side channels
-			DestinationFrame[0] += QuadChannelMap[0] * InputWithGain;
-			DestinationFrame[1] += QuadChannelMap[1] * InputWithGain;
-			DestinationFrame[5] += QuadChannelMap[2] * InputWithGain;
-			DestinationFrame[4] += QuadChannelMap[3] * InputWithGain;
-		}
-		else if (NumChannels >= 4)
-		{
-			// only really supporting quad, but weird channel configs will do something
-			for (int32 i = 0; i < 4; ++i)
+			Audio::GetStereoPan(ChannelAlpha, QuadChannelMap[Channel0], QuadChannelMap[Channel1]);
+
+			// Now map to the specific channel outputs
+			if (NumChannels == 6 || NumChannels == 8)
 			{
-				DestinationFrame[i] += QuadChannelMap[i] * InputWithGain;
+				// Specifically skipping LFE, Center channel, and side channels
+				DestinationFrame[0] += QuadChannelMap[0] * InputWithGain;
+				DestinationFrame[1] += QuadChannelMap[1] * InputWithGain;
+				DestinationFrame[5] += QuadChannelMap[2] * InputWithGain;
+				DestinationFrame[4] += QuadChannelMap[3] * InputWithGain;
+			}
+			else if (NumChannels >= 4)
+			{
+				// only really supporting quad, but weird channel configs will do something
+				for (int32 i = 0; i < 4; ++i)
+				{
+					DestinationFrame[i] += QuadChannelMap[i] * InputWithGain;
+				}
 			}
 		}
 	}
-}
 
-static void MixToMono(const float* InputFrame, const int32& NumChannels, float& Output)
-{
-	Output = 0.0f;
-
-	for (int32 Channel = 0; Channel < NumChannels; Channel++)
-	{
-		Output += InputFrame[Channel];
-	}
-}
-
-// Static id to get unique tap Ids
-static int32 TapIdCount = 0;
+	// Static id to get unique tap Ids
+	static int32 IdCount = 0;
+} // namespace TapDelayUtils
 
 FTapDelayInfo::FTapDelayInfo()
 	: TapLineMode(ETapLineMode::Panning)
@@ -73,7 +66,7 @@ FTapDelayInfo::FTapDelayInfo()
 	, Gain(-3.0f)
 	, OutputChannel(0)
 	, PanInDegrees(0.0f)
-	, TapId(TapIdCount++)
+	, TapId(TapDelayUtils::IdCount++)
 {
 }
 
@@ -84,10 +77,6 @@ FSubmixEffectTapDelay::FSubmixEffectTapDelay()
 {
 }
 
-FSubmixEffectTapDelay::~FSubmixEffectTapDelay()
-{
-}
-
 void FSubmixEffectTapDelay::Init(const FSoundEffectSubmixInitData& InData)
 {
 	SampleRate = InData.SampleRate;
@@ -124,7 +113,7 @@ void FSubmixEffectTapDelay::OnProcessAudio(const FSoundEffectSubmixInputData& In
 
 			if (TargetTap.TapLineMode == ETapLineMode::Panning || TargetTap.TapLineMode == ETapLineMode::Disabled)
 			{
-				GetMultichannelPan(DelayLine.ReadDelayAt(CurrentTap.GetLengthValue()), TargetTap.PanInDegrees, CurrentTap.GetGainValue(), OutData.NumChannels, &(OutBuffer[OutputBufferIndex]));
+				TapDelayUtils::GetMultichannelPan(DelayLine.ReadDelayAt(CurrentTap.GetLengthValue()), TargetTap.PanInDegrees, CurrentTap.GetGainValue(), OutData.NumChannels, &(OutBuffer[OutputBufferIndex]));
 			}
 			else
 			{
@@ -132,9 +121,14 @@ void FSubmixEffectTapDelay::OnProcessAudio(const FSoundEffectSubmixInputData& In
 				OutBuffer[BufferIndex] += DelayLine.ReadDelayAt(CurrentTap.GetLengthValue()) * CurrentTap.GetGainValue();
 			}
 		}
+
 		// Finally, write our input into the delay line.
-		float Input;
-		MixToMono(&InBuffer[InputBufferIndex], InData.NumChannels, Input);
+		float Input = 0.0f;
+		for (int32 Channel = 0; Channel < InData.NumChannels; Channel++)
+		{
+			Input += (&InBuffer[InputBufferIndex])[Channel];
+		}
+		
 		DelayLine.WriteDelayAndInc(Input);
 		InputBufferIndex += InData.NumChannels;
 	}
@@ -291,15 +285,12 @@ void FSubmixEffectTapDelay::UpdateInterpolations()
 void USubmixEffectTapDelayPreset::SetInterpolationTime(float Time)
 {
 	DynamicSettings.InterpolationTime = Time;
-	// Dispatch the added tap to all effect instances:
-	for (FSoundEffectBase* EffectBaseInstance : Instances)
+
+	// Dispatch the tap removal to all effect instances:
+	EffectCommand<FSubmixEffectTapDelay>([Time](FSubmixEffectTapDelay& TapDelay)
 	{
-		FSubmixEffectTapDelay* TapDelay = (FSubmixEffectTapDelay*)EffectBaseInstance;
-		EffectBaseInstance->EffectCommand([TapDelay, Time]()
-		{
-			TapDelay->SetInterpolationTime(Time);
-		});
-	}
+		TapDelay.SetInterpolationTime(Time);
+	});
 }
 
 void USubmixEffectTapDelayPreset::OnInit()
@@ -323,14 +314,10 @@ void USubmixEffectTapDelayPreset::AddTap(int32& OutTapId)
 	DynamicSettings.Taps.Add(TapInfo);
 
 	// Dispatch the added tap to all effect instances:
-	for (FSoundEffectBase* EffectBaseInstance : Instances)
+	EffectCommand<FSubmixEffectTapDelay>([OutTapId](FSubmixEffectTapDelay& TapDelay)
 	{
-		FSubmixEffectTapDelay* TapDelay = (FSubmixEffectTapDelay*)EffectBaseInstance;
-		EffectBaseInstance->EffectCommand([TapDelay, OutTapId]()
-		{
-			TapDelay->AddTap(OutTapId);
-		});
-	}
+		TapDelay.AddTap(OutTapId);
+	});
 }
 
 void USubmixEffectTapDelayPreset::RemoveTap(int32 TapId)
@@ -342,14 +329,10 @@ void USubmixEffectTapDelayPreset::RemoveTap(int32 TapId)
 			DynamicSettings.Taps.RemoveAtSwap(TapIndex);
 
 			// Dispatch the tap removal to all effect instances:
-			for (FSoundEffectBase* EffectBaseInstance : Instances)
+			EffectCommand<FSubmixEffectTapDelay>([TapId](FSubmixEffectTapDelay& TapDelay)
 			{
-				FSubmixEffectTapDelay* TapDelayInstance = (FSubmixEffectTapDelay*)EffectBaseInstance;
-				EffectBaseInstance->EffectCommand([TapDelayInstance, TapId]()
-				{
-					TapDelayInstance->RemoveTap(TapId);
-				});
-			}
+				TapDelay.RemoveTap(TapId);
+			});
 
 			return;
 		}
@@ -373,14 +356,11 @@ void USubmixEffectTapDelayPreset::SetTap(int32 TapId, const FTapDelayInfo& TapIn
 			DelayInfo.TapId = TapId;
 
 			// Dispatch the new tap to all effect instances:
-			for (FSoundEffectBase* EffectBaseInstance : Instances)
+			const FTapDelayInfo PassedTapInfo = TapInfo;
+			EffectCommand<FSubmixEffectTapDelay>([TapId, PassedTapInfo](FSubmixEffectTapDelay& TapDelay)
 			{
-				FSubmixEffectTapDelay* TapDelayInstance = (FSubmixEffectTapDelay*)EffectBaseInstance;
-				EffectBaseInstance->EffectCommand([TapDelayInstance, TapId, TapInfo]()
-				{
-					TapDelayInstance->SetTap(TapId, TapInfo);
-				});
-			}
+				TapDelay.SetTap(TapId, PassedTapInfo);
+			});
 
 			return;
 		}
diff --git a/Engine/Source/Runtime/AudioExtensions/Public/IAudioExtensionPlugin.h b/Engine/Source/Runtime/AudioExtensions/Public/IAudioExtensionPlugin.h
--- a/Engine/Source/Runtime/AudioExtensions/Public/IAudioExtensionPlugin.h
+++ b/Engine/Source/Runtime/AudioExtensions/Public/IAudioExtensionPlugin.h
@@ -10,19 +10,29 @@
 #include "UObject/ObjectMacros.h"
 #include "AudioDefines.h"
 
+#include "IAudioExtensionPlugin.generated.h"
+
+// Forward Declarations
+
 #if !UE_BUILD_SHIPPING
-class FCommonViewportClient;
 class FCanvas;
-class UFont;
+class FCommonViewportClient;
 class FViewport;
+class UFont;
 #endif // !UE_BUILD_SHIPPING
 
-#include "IAudioExtensionPlugin.generated.h"
-
-// Forward Declarations
 class FAudioDevice;
+class FSoundEffectBase;
+class FSoundEffectSource;
+class FSoundEffectSubmix;
+class IAudioModulation;
+class IAudioOcclusion;
+class IAudioPluginListener;
+class IAudioReverb;
+class IAudioSpatialization;
 class USoundSubmix;
 
+
 /**
 * Enumeration of audio plugin types
 *
@@ -37,14 +47,13 @@ enum class EAudioPlugin : uint8
 	COUNT = 4
 };
 
-class IAudioSpatialization;
-class IAudioModulation;
-class IAudioOcclusion;
-class IAudioReverb;
-class IAudioPluginListener;
-class FSoundEffectSubmix;
-typedef TSharedPtr<FSoundEffectSubmix, ESPMode::ThreadSafe> FSoundEffectSubmixPtr;
 
+// Deprecated in favor of TSoundEffectSubmixPtr
+using FSoundEffectSubmixPtr   = TSharedPtr<FSoundEffectSubmix, ESPMode::ThreadSafe>;
+
+using TSoundEffectPtr		  = TSharedPtr<FSoundEffectBase, ESPMode::ThreadSafe>;
+using TSoundEffectSourcePtr   = TSharedPtr<FSoundEffectSource, ESPMode::ThreadSafe>;
+using TSoundEffectSubmixPtr   = TSharedPtr<FSoundEffectSubmix, ESPMode::ThreadSafe>;
 using TAudioSpatializationPtr = TSharedPtr<IAudioSpatialization, ESPMode::ThreadSafe>;
 using TAudioModulationPtr     = TSharedPtr<IAudioModulation, ESPMode::ThreadSafe>;
 using TAudioOcclusionPtr      = TSharedPtr<IAudioOcclusion, ESPMode::ThreadSafe>;
diff --git a/Engine/Source/Runtime/AudioMixer/Private/AudioMixerBlueprintLibrary.cpp b/Engine/Source/Runtime/AudioMixer/Private/AudioMixerBlueprintLibrary.cpp
--- a/Engine/Source/Runtime/AudioMixer/Private/AudioMixerBlueprintLibrary.cpp
+++ b/Engine/Source/Runtime/AudioMixer/Private/AudioMixerBlueprintLibrary.cpp
@@ -9,6 +9,7 @@
 #include "ContentStreaming.h"
 #include "AudioCompressionSettingsUtils.h"
 #include "Async/Async.h"
+#include "Sound/SoundEffectPreset.h"
 
 // This is our global recording task:
 static TUniquePtr<Audio::FAudioRecordingData> RecordingData;
@@ -51,17 +52,13 @@ void UAudioMixerBlueprintLibrary::AddMasterSubmixEffect(const UObject* WorldCont
 
 	if (Audio::FMixerDevice* MixerDevice = GetAudioMixerDeviceFromWorldContext(WorldContextObject))
 	{
-		// Immediately create a new sound effect base here before the object becomes potentially invalidated
-		FSoundEffectSubmixPtr SoundEffectSubmix = MakeShareable(static_cast<FSoundEffectSubmix*>(SubmixEffectPreset->CreateNewEffect()));
-
 		FSoundEffectSubmixInitData InitData;
 		InitData.SampleRate = MixerDevice->GetSampleRate();
 		InitData.DeviceID = MixerDevice->DeviceID;
 		InitData.PresetSettings = nullptr;
 
-		// Initialize and set the preset immediately
-		SoundEffectSubmix->Init(InitData);
-		SoundEffectSubmix->SetPreset(SubmixEffectPreset);
+		// Immediately create a new sound effect base here before the object becomes potentially invalidated
+		TSoundEffectSubmixPtr SoundEffectSubmix = USoundEffectPreset::CreateInstance<FSoundEffectSubmixInitData, FSoundEffectSubmix>(InitData, *SubmixEffectPreset);
 		SoundEffectSubmix->SetEnabled(true);
 
 		// Get a unique id for the preset object on the game thread. Used to refer to the object on audio render thread.
diff --git a/Engine/Source/Runtime/AudioMixer/Private/AudioMixerSourceManager.cpp b/Engine/Source/Runtime/AudioMixer/Private/AudioMixerSourceManager.cpp
--- a/Engine/Source/Runtime/AudioMixer/Private/AudioMixerSourceManager.cpp
+++ b/Engine/Source/Runtime/AudioMixer/Private/AudioMixerSourceManager.cpp
@@ -85,7 +85,6 @@ FAutoConsoleVariableRef CVarCommandBufferFlushWaitTimeMs(
 	ECVF_Default);
 
 
-#define ONEOVERSHORTMAX (3.0517578125e-5f) // 1/32768
 #define ENVELOPE_TAIL_THRESHOLD (1.58489e-5f) // -96 dB
 
 #define VALIDATE_SOURCE_MIXER_STATE 1
@@ -555,6 +554,7 @@ namespace Audio
 	void FMixerSourceManager::BuildSourceEffectChain(const int32 SourceId, FSoundEffectSourceInitData& InitData, const TArray<FSourceEffectChainEntry>& InSourceEffectChain)
 	{
 		// Create new source effects. The memory will be owned by the source manager.
+		FScopeLock ScopeLock(&EffectChainMutationCriticalSection);
 		for (const FSourceEffectChainEntry& ChainEntry : InSourceEffectChain)
 		{
 			// Presets can have null entries
@@ -563,14 +563,11 @@ namespace Audio
 				continue;
 			}
 
-			FSoundEffectSource* NewEffect = static_cast<FSoundEffectSource*>(ChainEntry.Preset->CreateNewEffect());
-
 			// Get this source effect presets unique id so instances can identify their originating preset object
 			const uint32 PresetUniqueId = ChainEntry.Preset->GetUniqueID();
 			InitData.ParentPresetUniqueId = PresetUniqueId;
 
-			NewEffect->Init(InitData);
-			NewEffect->SetPreset(ChainEntry.Preset);
+			TSoundEffectSourcePtr NewEffect = USoundEffectPreset::CreateInstance<FSoundEffectSourceInitData, FSoundEffectSource>(InitData, *ChainEntry.Preset);
 			NewEffect->SetEnabled(!ChainEntry.bBypass);
 
 			// Add the effect instance
@@ -584,21 +581,22 @@ namespace Audio
 
 	void FMixerSourceManager::ResetSourceEffectChain(const int32 SourceId)
 	{
-		FSourceInfo& SourceInfo = SourceInfos[SourceId];
-
-		for (int32 i = 0; i < SourceInfo.SourceEffects.Num(); ++i)
+		FScopeLock ScopeLock(&EffectChainMutationCriticalSection);
 		{
-			SourceInfo.SourceEffects[i]->ClearPreset();
-			delete SourceInfo.SourceEffects[i];
-			SourceInfo.SourceEffects[i] = nullptr;
-		}
-		SourceInfo.SourceEffects.Reset();
+			FSourceInfo& SourceInfo = SourceInfos[SourceId];
 
-		for (int32 i = 0; i < SourceInfo.SourceEffectPresets.Num(); ++i)
-		{
-			SourceInfo.SourceEffectPresets[i] = nullptr;
+			for (int32 i = 0; i < SourceInfo.SourceEffects.Num(); ++i)
+			{
+				USoundEffectPreset::UnregisterInstance(SourceInfo.SourceEffects[i]);
+			}
+			SourceInfo.SourceEffects.Reset();
+
+			for (int32 i = 0; i < SourceInfo.SourceEffectPresets.Num(); ++i)
+			{
+				SourceInfo.SourceEffectPresets[i] = nullptr;
+			}
+			SourceInfo.SourceEffectPresets.Reset();
 		}
-		SourceInfo.SourceEffectPresets.Reset();
 	}
 
 	bool FMixerSourceManager::GetFreeSourceId(int32& OutSourceId)
@@ -2239,27 +2237,30 @@ namespace Audio
 				SourceInfo.SourceEffectInputData.NumSamples = NumSamples;
 
 				// Loop through the effect chain passing in buffers
-				for (FSoundEffectSource* SoundEffectSource : SourceInfo.SourceEffects)
+				FScopeLock ScopeLock(&EffectChainMutationCriticalSection);
 				{
-					bool bPresetUpdated = false;
-					if (SoundEffectSource->IsActive())
+					for (TSoundEffectSourcePtr& SoundEffectSource : SourceInfo.SourceEffects)
 					{
-						bPresetUpdated = SoundEffectSource->Update();
-					}
+						bool bPresetUpdated = false;
+						if (SoundEffectSource->IsActive())
+						{
+							bPresetUpdated = SoundEffectSource->Update();
+						}
 
-					// Modulation must be updated regardless of whether or not the source is
-					// active to allow for initial conditions to be set if source is reactivated.
-					if (bPresetUpdated || SourceInfo.bIsModulationUpdated)
-					{
-						SoundEffectSource->ProcessControls(SourceInfo.ModulationControls);
-					}
+						// Modulation must be updated regardless of whether or not the source is
+						// active to allow for initial conditions to be set if source is reactivated.
+						if (bPresetUpdated || SourceInfo.bIsModulationUpdated)
+						{
+							SoundEffectSource->ProcessControls(SourceInfo.ModulationControls);
+						}
 
-					if (SoundEffectSource->IsActive())
-					{
-						SoundEffectSource->ProcessAudio(SourceInfo.SourceEffectInputData, OutputSourceEffectBufferPtr);
+						if (SoundEffectSource->IsActive())
+						{
+							SoundEffectSource->ProcessAudio(SourceInfo.SourceEffectInputData, OutputSourceEffectBufferPtr);
 
-						// Copy output to input
-						FMemory::Memcpy(SourceInfo.SourceEffectInputData.InputSourceEffectBufferPtr, OutputSourceEffectBufferPtr, sizeof(float)*NumSamples);
+							// Copy output to input
+							FMemory::Memcpy(SourceInfo.SourceEffectInputData.InputSourceEffectBufferPtr, OutputSourceEffectBufferPtr, sizeof(float) * NumSamples);
+						}
 					}
 				}
 
@@ -2553,40 +2554,43 @@ namespace Audio
 					SourceInfo.bEffectTailsDone = !bPlayEffectChainTails;
 
 					// Check to see if the chain didn't actually change
-					TArray<FSoundEffectSource *>& ThisSourceEffectChain = SourceInfo.SourceEffects;
-					bool bReset = false;
-					if (InSourceEffectChain.Num() == ThisSourceEffectChain.Num())
+					FScopeLock ScopeLock(&EffectChainMutationCriticalSection);
 					{
-						for (int32 SourceEffectId = 0; SourceEffectId < ThisSourceEffectChain.Num(); ++SourceEffectId)
+						TArray<TSoundEffectSourcePtr>& ThisSourceEffectChain = SourceInfo.SourceEffects;
+						bool bReset = false;
+						if (InSourceEffectChain.Num() == ThisSourceEffectChain.Num())
 						{
-							const FSourceEffectChainEntry& ChainEntry = InSourceEffectChain[SourceEffectId];
-
-							FSoundEffectSource* SourceEffectInstance = ThisSourceEffectChain[SourceEffectId];
-							if (!SourceEffectInstance->IsPreset(ChainEntry.Preset))
+							for (int32 SourceEffectId = 0; SourceEffectId < ThisSourceEffectChain.Num(); ++SourceEffectId)
 							{
-								// As soon as one of the effects change or is not the same, then we need to rebuild the effect graph
-								bReset = true;
-								break;
+								const FSourceEffectChainEntry& ChainEntry = InSourceEffectChain[SourceEffectId];
+
+								TSoundEffectSourcePtr SourceEffectInstance = ThisSourceEffectChain[SourceEffectId];
+								if (!SourceEffectInstance->IsPreset(ChainEntry.Preset))
+								{
+									// As soon as one of the effects change or is not the same, then we need to rebuild the effect graph
+									bReset = true;
+									break;
+								}
+
+								// Otherwise just update if it's just to bypass
+								SourceEffectInstance->SetEnabled(!ChainEntry.bBypass);
 							}
-
-							// Otherwise just update if it's just to bypass
-							SourceEffectInstance->SetEnabled(!ChainEntry.bBypass);
 						}
-					}
-					else
-					{
-						bReset = true;
-					}
+						else
+						{
+							bReset = true;
+						}
 
-					if (bReset)
-					{
-						InitData.NumSourceChannels = SourceInfo.NumInputChannels;
+						if (bReset)
+						{
+							InitData.NumSourceChannels = SourceInfo.NumInputChannels;
 
-						// First reset the source effect chain
-						ResetSourceEffectChain(SourceId);
+							// First reset the source effect chain
+							ResetSourceEffectChain(SourceId);
 
-						// Rebuild it
-						BuildSourceEffectChain(SourceId, InitData, InSourceEffectChain);
+							// Rebuild it
+							BuildSourceEffectChain(SourceId, InitData, InSourceEffectChain);
+						}
 					}
 				}
 			}
diff --git a/Engine/Source/Runtime/AudioMixer/Private/AudioMixerSourceManager.h b/Engine/Source/Runtime/AudioMixer/Private/AudioMixerSourceManager.h
--- a/Engine/Source/Runtime/AudioMixer/Private/AudioMixerSourceManager.h
+++ b/Engine/Source/Runtime/AudioMixer/Private/AudioMixerSourceManager.h
@@ -303,6 +303,9 @@ namespace Audio
 			}
 		};
 
+		// Critical section to ensure mutating effect chains is thread-safe
+		FCriticalSection EffectChainMutationCriticalSection;
+
 		FMixerDevice* MixerDevice;
 
 		// Cached ptr to an optional spatialization plugin
@@ -476,7 +479,7 @@ namespace Audio
 
 			// Source effect instances
 			uint32 SourceEffectChainId;
-			TArray<FSoundEffectSource*> SourceEffects;
+			TArray<TSoundEffectSourcePtr> SourceEffects;
 			TArray<USoundEffectSourcePreset*> SourceEffectPresets;
 			bool bEffectTailsDone;
 			FSoundEffectSourceInputData SourceEffectInputData;
diff --git a/Engine/Source/Runtime/AudioMixer/Private/AudioMixerSubmix.cpp b/Engine/Source/Runtime/AudioMixer/Private/AudioMixerSubmix.cpp
--- a/Engine/Source/Runtime/AudioMixer/Private/AudioMixerSubmix.cpp
+++ b/Engine/Source/Runtime/AudioMixer/Private/AudioMixerSubmix.cpp
@@ -3,6 +3,7 @@
 #include "AudioMixerSubmix.h"
 #include "AudioMixerDevice.h"
 #include "AudioMixerSourceVoice.h"
+#include "Sound/SoundEffectPreset.h"
 #include "Sound/SoundEffectSubmix.h"
 #include "Sound/SoundSubmix.h"
 #include "Sound/SoundSubmixSend.h"
@@ -104,8 +105,6 @@ namespace Audio
 			CurrentOutputVolume = FMath::Clamp(SoundSubmix->OutputVolume, 0.0f, 1.0f);
 			TargetOutputVolume = CurrentOutputVolume;
 
-
-			FScopeLock ScopeLock(&EffectChainMutationCriticalSection);
 			// Set the initialized output volume
 			CurrentWetLevel = FMath::Clamp(SoundSubmix->WetLevel, 0.0f, 1.0f);
 			TargetWetLevel = CurrentWetLevel;
@@ -113,33 +112,32 @@ namespace Audio
 			CurrentDryLevel = FMath::Clamp(SoundSubmix->DryLevel, 0.0f, 1.0f);
 			TargetDryLevel = CurrentDryLevel;
 
-			NumSubmixEffects = 0;
-
-			for (USoundEffectSubmixPreset* EffectPreset : SoundSubmix->SubmixEffectChain)
+			FScopeLock ScopeLock(&EffectChainMutationCriticalSection);
 			{
-				if (EffectPreset)
+    			NumSubmixEffects = 0;
+				for (USoundEffectSubmixPreset* EffectPreset : SoundSubmix->SubmixEffectChain)
 				{
-					++NumSubmixEffects;
+					if (EffectPreset)
+					{
+						++NumSubmixEffects;
 
-					// Create a new effect instance using the preset
-					FSoundEffectSubmix* SubmixEffect = static_cast<FSoundEffectSubmix*>(EffectPreset->CreateNewEffect());
 
-					FSoundEffectSubmixInitData InitData;
-					InitData.DeviceID = MixerDevice->DeviceID;
-					InitData.SampleRate = MixerDevice->GetSampleRate();
-					InitData.PresetSettings = nullptr;
+						FSoundEffectSubmixInitData InitData;
+						InitData.DeviceID = MixerDevice->DeviceID;
+						InitData.SampleRate = MixerDevice->GetSampleRate();
+						InitData.PresetSettings = nullptr;
 
-					// Now set the preset
-					SubmixEffect->Init(InitData);
-					SubmixEffect->SetPreset(EffectPreset);
-					SubmixEffect->SetEnabled(true);
+						// Create a new effect instance using the preset & enable
+						TSoundEffectSubmixPtr SubmixEffect = USoundEffectPreset::CreateInstance<FSoundEffectSubmixInitData, FSoundEffectSubmix>(InitData, *EffectPreset);
+						SubmixEffect->SetEnabled(true);
 
-					FSubmixEffectInfo EffectInfo;
-					EffectInfo.PresetId = EffectPreset->GetUniqueID();
-					EffectInfo.EffectInstance = MakeShareable(SubmixEffect);
+						FSubmixEffectInfo EffectInfo;
+						EffectInfo.PresetId = EffectPreset->GetUniqueID();
+						EffectInfo.EffectInstance = SubmixEffect;
 
-					// Add the effect to this submix's chain
-					EffectSubmixChain.Add(MoveTemp(EffectInfo));
+						// Add the effect to this submix's chain
+						EffectSubmixChain.Add(MoveTemp(EffectInfo));
+					}
 				}
 			}
 
@@ -413,7 +411,7 @@ namespace Audio
 		{
 			if (Info.EffectInstance.IsValid())
 			{
-				Info.EffectInstance->ClearPreset();
+				USoundEffectPreset::UnregisterInstance(Info.EffectInstance);
 				Info.EffectInstance.Reset();
 			}
 		}
diff --git a/Engine/Source/Runtime/Engine/Classes/Sound/SoundEffectBase.h b/Engine/Source/Runtime/Engine/Classes/Sound/SoundEffectBase.h
--- a/Engine/Source/Runtime/Engine/Classes/Sound/SoundEffectBase.h
+++ b/Engine/Source/Runtime/Engine/Classes/Sound/SoundEffectBase.h
@@ -6,11 +6,14 @@
 #include "UObject/ObjectMacros.h"
 #include "UObject/Object.h"
 #include "HAL/ThreadSafeBool.h"
-#include "Sound/SoundEffectPreset.h"
 #include "Containers/Queue.h"
 #include "Misc/ScopeLock.h"
 
 
+// Forward Declarations
+class USoundEffectPreset;
+
+
 // The following macro code creates boiler-plate code for a sound effect preset and hides unnecessary details from user-created effects.
 
 // Macro chain to expand "MyEffectName" to "FMyEffectNameSettings"
@@ -57,14 +60,10 @@
 		FCriticalSection SettingsCritSect; \
 		F##EFFECT_NAME##Settings SettingsCopy; \
 
-class USoundEffectPreset;
-
 class ENGINE_API FSoundEffectBase
 {
 public:
-	FSoundEffectBase();
-
-	virtual ~FSoundEffectBase();
+	virtual ~FSoundEffectBase() = default;
 
 	/** Called when the sound effect's preset changed. */
 	virtual void OnPresetChanged() {};
@@ -78,13 +77,8 @@ public:
 	/** Updates preset on audio render thread. Returns true if update processed a preset update, false if not. */
 	bool Update();
 
-	void SetPreset(USoundEffectPreset* Inpreset);
-
 	USoundEffectPreset* GetPreset();
 
-	/** Removes the instance from the preset. */
-	void ClearPreset(bool bRemoveFromPreset = true);
-
 	/** Queries if the given preset object is the uobject preset for this preset instance, i.e. the preset which spawned this effect instance. */
 	bool IsPreset(USoundEffectPreset* InPreset) const;
 
@@ -92,7 +86,12 @@ public:
 	void EffectCommand(TFunction<void()> Command);
 
 protected:
+	FSoundEffectBase();
 
+	/** Removes the instance from the preset. */
+	void ClearPreset();
+
+	/** Pumps messages awaiting execution on the audio render thread */
 	void PumpPendingMessages();
 
 	FCriticalSection SettingsCritSect;
@@ -104,11 +103,11 @@ protected:
 	FThreadSafeBool bIsRunning;
 	FThreadSafeBool bIsActive;
 
-	// Effect commmand queue
+	// Effect command queue
 	TQueue<TFunction<void()>> CommandQueue;
 
-	// Allow FAudioMixerSubmix to call ProcessAudio
-	friend class FMixerSubmix;
-
+	// Allow preset to re-register when editor update is requested
+	// and create effects using the templated Create call
+	friend class USoundEffectPreset;
 };
 
diff --git a/Engine/Source/Runtime/Engine/Classes/Sound/SoundEffectPreset.h b/Engine/Source/Runtime/Engine/Classes/Sound/SoundEffectPreset.h
--- a/Engine/Source/Runtime/Engine/Classes/Sound/SoundEffectPreset.h
+++ b/Engine/Source/Runtime/Engine/Classes/Sound/SoundEffectPreset.h
@@ -2,15 +2,22 @@
 #pragma once
 
 #include "CoreMinimal.h"
-#include "UObject/UObjectGlobals.h"
-#include "UObject/ObjectMacros.h"
+#include "IAudioExtensionPlugin.h"
+#include "Templates/SharedPointer.h"
 #include "UObject/Object.h"
+#include "UObject/ObjectMacros.h"
+#include "UObject/UObjectGlobals.h"
+#include "SoundEffectBase.h"
+
 #include "SoundEffectPreset.generated.h"
 
+// Forward Declarations
 struct FAssetData;
+
 class FMenuBuilder;
-class IToolkitHost;
 class FSoundEffectBase;
+class IToolkitHost;
+
 
 UCLASS(config = Engine, abstract, editinlinenew, BlueprintType)
 class ENGINE_API USoundEffectPreset : public UObject
@@ -19,7 +26,8 @@ class ENGINE_API USoundEffectPreset : public UObject
 
 public:
 	USoundEffectPreset(const FObjectInitializer& ObjectInitializer);
-	virtual ~USoundEffectPreset();
+	virtual ~USoundEffectPreset() = default;
+
 
 	virtual bool CanFilter() const { return true; }
 	virtual FText GetAssetActionName() const PURE_VIRTUAL(USoundEffectPreset::GetAssetActionName, return FText(););
@@ -31,34 +39,98 @@ public:
 	virtual void OnInit() {};
 	virtual FColor GetPresetColor() const { return FColor(200.0f, 100.0f, 100.0f); }
 
-	void EffectCommand(TFunction<void()> Command);
 	void Update();
-	void AddEffectInstance(FSoundEffectBase* InSource);
-	void RemoveEffectInstance(FSoundEffectBase* InSource);
+	void AddEffectInstance(TSoundEffectPtr& InEffectPtr);
+	void RemoveEffectInstance(TSoundEffectPtr& InEffectPtr);
 
 	void AddReferencedEffects(FReferenceCollector& InCollector);
 
 	virtual void BeginDestroy() override;
 
 protected:
-
 #if WITH_EDITORONLY_DATA
 	virtual void PostEditChangeProperty(FPropertyChangedEvent& PropertyChangedEvent) override;
 #endif
 
+	using TSoundEffectWeakPtr = TWeakPtr<FSoundEffectBase, ESPMode::ThreadSafe>;
+
 	// Array of instances which are using this preset
-	TArray<FSoundEffectBase*> Instances;
+	TArray<TSoundEffectWeakPtr> Instances;
 	bool bInitialized;
 
+	/* Immediately executes command for each active effect instance on the active thread */
 	template <typename T>
 	void IterateEffects(TFunction<void(T&)> InForEachEffect)
 	{
-		for (FSoundEffectBase* Instance : Instances)
+		for (TSoundEffectWeakPtr& Instance : Instances)
 		{
-			if (Instance)
+			TSoundEffectPtr EffectStrongPtr = Instance.Pin();
+			if (EffectStrongPtr.IsValid())
 			{
-				InForEachEffect(*static_cast<T*>(Instance));
+				InForEachEffect(*static_cast<T*>(EffectStrongPtr.Get()));
 			}
 		}
 	}
+
+	/* Defers execution of command on each active effect instance on the audio render thread */
+	template <typename T>
+	void EffectCommand(TFunction<void(T&)> InForEachEffect)
+	{
+		IterateEffects<T>([InForEachEffect](T& OutInstance)
+		{
+			T* InstancePtr = &OutInstance;
+			OutInstance.EffectCommand([InstancePtr, InForEachEffect]()
+			{
+				InForEachEffect(*InstancePtr);
+			});
+		});
+	}
+
+public:
+	template <typename TInitData, typename TSoundEffectType>
+	static TSharedPtr<TSoundEffectType, ESPMode::ThreadSafe> CreateInstance(const TInitData& InInitData, USoundEffectPreset& InOutPreset)
+	{
+		TSoundEffectType* NewEffect = static_cast<TSoundEffectType*>(InOutPreset.CreateNewEffect());
+		NewEffect->Init(InInitData);
+
+		TSharedPtr<TSoundEffectType, ESPMode::ThreadSafe> NewEffectPtr(NewEffect);
+		RegisterInstance(InOutPreset, StaticCastSharedPtr<FSoundEffectBase, TSoundEffectType, ESPMode::ThreadSafe>(NewEffectPtr));
+		return NewEffectPtr;
+	}
+
+	static void UnregisterInstance(TSoundEffectPtr InEffectPtr)
+	{
+		if (InEffectPtr.IsValid())
+		{
+			if (USoundEffectPreset* Preset = InEffectPtr->GetPreset())
+			{
+				Preset->RemoveEffectInstance(InEffectPtr);
+			}
+
+			InEffectPtr->ClearPreset();
+		}
+	}
+
+	static void RegisterInstance(USoundEffectPreset& InPreset, TSoundEffectPtr InEffectPtr)
+	{
+		if (!InEffectPtr.IsValid())
+		{
+			return;
+		}
+
+		if (InEffectPtr->Preset.Get() != &InPreset)
+		{
+			UnregisterInstance(InEffectPtr);
+
+			InEffectPtr->Preset = &InPreset;
+			if (InEffectPtr->Preset.IsValid())
+			{
+				InPreset.AddEffectInstance(InEffectPtr);
+			}
+		}
+
+		// Anytime notification occurs that the preset has been modified,
+		// flag for update.
+		InEffectPtr->bChanged = true;
+	}
 };
\ No newline at end of file
diff --git a/Engine/Source/Runtime/Engine/Classes/Sound/SoundEffectSource.h b/Engine/Source/Runtime/Engine/Classes/Sound/SoundEffectSource.h
--- a/Engine/Source/Runtime/Engine/Classes/Sound/SoundEffectSource.h
+++ b/Engine/Source/Runtime/Engine/Classes/Sound/SoundEffectSource.h
@@ -2,10 +2,11 @@
 #pragma once
 
 #include "CoreMinimal.h"
-#include "UObject/ObjectMacros.h"
+#include "IAudioExtensionPlugin.h"
 #include "Sound/SoundEffectPreset.h"
 #include "Sound/SoundEffectBase.h"
-#include "IAudioExtensionPlugin.h"
+#include "Templates/SharedPointer.h"
+#include "UObject/ObjectMacros.h"
 
 #include "SoundEffectSource.generated.h"
 
@@ -103,15 +104,18 @@ struct FSoundEffectSourceInputData
 
 class ENGINE_API FSoundEffectSource : public FSoundEffectBase
 {
-public:
-	virtual ~FSoundEffectSource() {}
-
+protected:
 	/** Called on an audio effect at initialization on main thread before audio processing begins. */
-	virtual void Init(const FSoundEffectSourceInitData& InSampleRate) = 0;
+	virtual void Init(const FSoundEffectSourceInitData& InInitData) = 0;
+
+public:
+	virtual ~FSoundEffectSource() = default;
 
 	/** Process the input block of audio. Called on audio thread. */
 	virtual void ProcessAudio(const FSoundEffectSourceInputData& InData, float* OutAudioBufferData) = 0;
 
 	/** Process modulation controls if enabled */
 	virtual void ProcessControls(const FSoundModulationControls& InControls) { }
+
+	friend class USoundEffectPreset;
 };
diff --git a/Engine/Source/Runtime/Engine/Classes/Sound/SoundEffectSubmix.h b/Engine/Source/Runtime/Engine/Classes/Sound/SoundEffectSubmix.h
--- a/Engine/Source/Runtime/Engine/Classes/Sound/SoundEffectSubmix.h
+++ b/Engine/Source/Runtime/Engine/Classes/Sound/SoundEffectSubmix.h
@@ -4,18 +4,21 @@
 #include "CoreMinimal.h"
 
 #include "AudioDefines.h"
+#include "IAudioExtensionPlugin.h"
 #include "Sound/SoundEffectPreset.h"
 #include "Sound/SoundEffectBase.h"
+#include "Templates/SharedPointer.h"
 #include "UObject/ObjectMacros.h"
 
 #include "SoundEffectSubmix.generated.h"
 
-
+// Forward Declarations
 class FAudioDevice;
 class FSoundEffectSubmix;
 
 struct FAudioEffectParameters;
 
+
 /** Preset of a submix effect that can be shared between sounds. */
 UCLASS(config = Engine, hidecategories = Object, abstract, editinlinenew, BlueprintType)
 class ENGINE_API USoundEffectSubmixPreset : public USoundEffectPreset
@@ -87,18 +90,7 @@ struct FSoundEffectSubmixOutputData
 class ENGINE_API FSoundEffectSubmix : public FSoundEffectBase
 {
 public:
-	FSoundEffectSubmix()
-	{
-	}
-
-	virtual ~FSoundEffectSubmix()
-	{
-	}
-
-	/** Called on an audio effect at initialization on main thread before audio processing begins. */
-	virtual void Init(const FSoundEffectSubmixInitData& InSampleRate)
-	{
-	}
+	virtual ~FSoundEffectSubmix() = default;
 
 	/**  Provided for interpolating parameters from audio volume system, enabling transition between various settings */
 	virtual bool SetParameters(const FAudioEffectParameters& InParameters)
@@ -139,5 +131,16 @@ public:
 
 	/** Processes audio in the submix effect. */
 	void ProcessAudio(FSoundEffectSubmixInputData& InData, FSoundEffectSubmixOutputData& OutData);
+
+	friend class USoundEffectPreset;
+
+protected:
+	FSoundEffectSubmix()
+	{
+	}
+
+	/** Called on an audio effect at initialization on main thread before audio processing begins. */
+	virtual void Init(const FSoundEffectSubmixInitData& InInitData)
+	{
+	}
 };
-typedef TSharedPtr<FSoundEffectSubmix, ESPMode::ThreadSafe> FSoundEffectSubmixPtr;
diff --git a/Engine/Source/Runtime/Engine/Classes/Sound/SoundSubmix.h b/Engine/Source/Runtime/Engine/Classes/Sound/SoundSubmix.h
--- a/Engine/Source/Runtime/Engine/Classes/Sound/SoundSubmix.h
+++ b/Engine/Source/Runtime/Engine/Classes/Sound/SoundSubmix.h
@@ -229,10 +229,6 @@ public:
 	UFUNCTION(BlueprintCallable, Category = "Audio", meta = (WorldContext = "WorldContextObject"))
 	void SetSubmixOutputVolume(const UObject* WorldContextObject, float InOutputVolume);
 
-	// Registers and unregisters buffer listeners with the submix
-	void RegisterSubmixBufferListener(ISubmixBufferListener* InBufferListener);
-	void UnregisterSubmixBufferListener(ISubmixBufferListener* InBufferListener);
-
 protected:
 
 	virtual void PostLoad() override;
diff --git a/Engine/Source/Runtime/Engine/Private/SoundEffectBase.cpp b/Engine/Source/Runtime/Engine/Private/SoundEffectBase.cpp
--- a/Engine/Source/Runtime/Engine/Private/SoundEffectBase.cpp
+++ b/Engine/Source/Runtime/Engine/Private/SoundEffectBase.cpp
@@ -1,7 +1,8 @@
 // Copyright Epic Games, Inc. All Rights Reserved.
 
-
 #include "Sound/SoundEffectBase.h"
+#include "Sound/SoundEffectPreset.h"
+#include "Templates/SharedPointer.h"
 
 
 FSoundEffectBase::FSoundEffectBase()
@@ -10,10 +11,6 @@ FSoundEffectBase::FSoundEffectBase()
 	, bIsActive(false)
 {}
 
-FSoundEffectBase::~FSoundEffectBase()
-{
-}
-
 bool FSoundEffectBase::IsActive() const
 {
 	return bIsActive;
@@ -24,36 +21,13 @@ void FSoundEffectBase::SetEnabled(const bool bInIsEnabled)
 	bIsActive = bInIsEnabled;
 }
 
-void FSoundEffectBase::SetPreset(USoundEffectPreset* Inpreset)
-{
-	if (Preset != Inpreset)
-	{
-		ClearPreset();
-
-		Preset = Inpreset;
-		if (Preset.IsValid())
-		{
-			Preset->AddEffectInstance(this);
-		}
-	}
-
-	// Anytime notification occurs that the preset has been modified,
-	// flag for update.
-	bChanged = true;
-}
-
 USoundEffectPreset* FSoundEffectBase::GetPreset()
 {
 	return Preset.Get();
 }
 
-void FSoundEffectBase::ClearPreset(bool bRemoveFromPreset)
+void FSoundEffectBase::ClearPreset()
 {
-	if (bRemoveFromPreset && Preset.IsValid())
-	{
-		Preset->RemoveEffectInstance(this);
-	}
-
 	Preset.Reset();
 }
 
diff --git a/Engine/Source/Runtime/Engine/Private/SoundEffectPreset.cpp b/Engine/Source/Runtime/Engine/Private/SoundEffectPreset.cpp
--- a/Engine/Source/Runtime/Engine/Private/SoundEffectPreset.cpp
+++ b/Engine/Source/Runtime/Engine/Private/SoundEffectPreset.cpp
@@ -14,34 +14,23 @@ USoundEffectPreset::USoundEffectPreset(const FObjectInitializer& ObjectInitializ
 
 }
 
-USoundEffectPreset::~USoundEffectPreset()
-{
-}
-
-void USoundEffectPreset::EffectCommand(TFunction<void()> Command)
-{
-	IterateEffects<FSoundEffectBase>([Command](FSoundEffectBase& Instance)
-	{
-		Instance.EffectCommand(Command);
-	});
-}
-
 void USoundEffectPreset::Update()
 {
 	for (int32 i = Instances.Num() - 1; i >= 0; --i)
 	{
-		if (!Instances[i] || Instances[i]->GetPreset() == nullptr)
+		TSoundEffectPtr EffectSharedPtr = Instances[i].Pin();
+		if (!EffectSharedPtr.IsValid() || EffectSharedPtr->GetPreset() == nullptr)
 		{
-			Instances.RemoveAtSwap(i, 1, false /* bAllowShrinking */);
+			Instances.RemoveAtSwap(i, 1);
 		}
 		else
 		{
-			Instances[i]->SetPreset(this);
+			RegisterInstance(*this, EffectSharedPtr);
 		}
 	}
 }
 
-void USoundEffectPreset::AddEffectInstance(FSoundEffectBase* InSource)
+void USoundEffectPreset::AddEffectInstance(TSoundEffectPtr& InEffectPtr)
 {
 	if (!bInitialized)
 	{
@@ -52,7 +41,7 @@ void USoundEffectPreset::AddEffectInstance(FSoundEffectBase* InSource)
 		OnInit();
 	}
 
-	Instances.AddUnique(InSource);
+	Instances.AddUnique(TSoundEffectWeakPtr(InEffectPtr));
 }
 
 void USoundEffectPreset::AddReferencedEffects(FReferenceCollector& InCollector)
@@ -60,28 +49,27 @@ void USoundEffectPreset::AddReferencedEffects(FReferenceCollector& InCollector)
 	FReferenceCollector* Collector = &InCollector;
 	IterateEffects<FSoundEffectBase>([Collector](FSoundEffectBase& Instance)
 	{
-		const USoundEffectPreset* EffectPreset = Instance.GetPreset();
-		Collector->AddReferencedObject(EffectPreset);
+		if (const USoundEffectPreset* EffectPreset = Instance.GetPreset())
+		{
+			Collector->AddReferencedObject(EffectPreset);
+		}
 	});
 }
 
 void USoundEffectPreset::BeginDestroy()
 {
-	if (!GExitPurge)
+	IterateEffects<FSoundEffectBase>([](FSoundEffectBase& Instance)
 	{
-		IterateEffects<FSoundEffectBase>([](FSoundEffectBase& Instance)
-		{
-			Instance.ClearPreset(false /* bRemoveFromPreset */);
-		});
-	}
+		Instance.ClearPreset();
+	});
 	Instances.Reset();
 
 	Super::BeginDestroy();
 }
 
-void USoundEffectPreset::RemoveEffectInstance(FSoundEffectBase* InSource)
+void USoundEffectPreset::RemoveEffectInstance(TSoundEffectPtr& InEffectPtr)
 {
-	Instances.RemoveSwap(InSource);
+	Instances.RemoveSwap(TSoundEffectWeakPtr(InEffectPtr));
 }
 
 #if WITH_EDITORONLY_DATA
