Reworking convolution rever submix effect

diff --git a/Engine/Plugins/Runtime/Synthesis/Source/Synthesis/Classes/SubmixEffects/SubmixEffectConvolutionReverb.h b/Engine/Plugins/Runtime/Synthesis/Source/Synthesis/Classes/SubmixEffects/SubmixEffectConvolutionReverb.h
new file mode 100644
--- /dev/null
+++ b/Engine/Plugins/Runtime/Synthesis/Source/Synthesis/Classes/SubmixEffects/SubmixEffectConvolutionReverb.h
@@ -0,0 +1,296 @@
+// Copyright Epic Games, Inc. All Rights Reserved.
+#pragma once
+
+#include "CoreMinimal.h"
+#include "UObject/Object.h"
+#include "ConvolutionReverb.h"
+#include "DSP/Dsp.h"
+#include "DSP/ConvolutionAlgorithm.h"
+#include "Sound/SoundEffectSubmix.h"
+#include "SubmixEffectConvolutionReverb.generated.h"
+
+/** Block size of convolution reverb algorithm. */
+UENUM()
+enum class ESubmixEffectConvolutionReverbBlockSize : uint8
+{
+	/** 256 audio frames per a block. */
+	BlockSize256,
+
+	/** 512 audio frames per a block. */
+	BlockSize512,
+
+	/** 1024 audio frames per a block. */
+	BlockSize1024
+};
+
+#if WITH_EDITORONLY_DATA
+DECLARE_MULTICAST_DELEGATE_OneParam(FAudioImpulseResponsePropertyChange, FPropertyChangedEvent&);
+#endif
+
+// ========================================================================
+// UAudioImpulseResponse
+// UAsset used to represent Imported Impulse Responses
+// ========================================================================
+UCLASS()
+class SYNTHESIS_API UAudioImpulseResponse : public UObject
+{
+	GENERATED_BODY()
+
+public:
+	UAudioImpulseResponse();
+
+	/** The interleaved audio samples used in convolution. */
+	UPROPERTY()
+	TArray<float> ImpulseResponse;
+
+	/** The number of channels in impulse response. */
+	UPROPERTY()
+	int32 NumChannels;
+
+	/** The original sample rate of the impulse response. */
+	UPROPERTY()
+	int32 SampleRate;
+
+	/* Used to account for energy added by convolution with "loud" Impulse Responses. */
+	UPROPERTY(EditAnywhere, BlueprintReadWrite, Category = SubmixEffectPreset, meta = (ClamMin = "-60.0", ClampMax = "15.0"))
+	float NormalizationVolumeDb;
+
+#if WITH_EDITORONLY_DATA
+	virtual void PostEditChangeProperty(FPropertyChangedEvent& PropertyChangedEvent) override;
+
+	// This delegate is called whenever PostEditChangeProperty called.
+	FAudioImpulseResponsePropertyChange OnObjectPropertyChanged;
+#endif
+};
+
+
+USTRUCT(BlueprintType)
+struct SYNTHESIS_API FSubmixEffectConvolutionReverbSettings
+{
+	GENERATED_USTRUCT_BODY()
+
+	FSubmixEffectConvolutionReverbSettings();
+
+	/* Used to account for energy added by convolution with "loud" Impulse Responses. 
+	 * This value is not directly editable in the editor because it is copied from the 
+	 * associated UAudioImpulseResponse. */
+	UPROPERTY();
+	float NormalizationVolumeDb;
+
+	/* Amout of audio to be sent to rear channels in quad/surround configurations */
+	UPROPERTY(EditAnywhere, BlueprintReadWrite, Category = SubmixEffectPreset, meta = (ClampMin = "-60.0", UIMin = "-60.0", ClampMax = "15.0", UIMax = "15.0"))
+	float SurroundRearChannelBleedDb;
+
+	/* If true, rear channel bleed sends will have their phase inverted. */
+	UPROPERTY(EditAnywhere, BlueprintReadWrite, Category = SumixEffectPreset)
+	bool bInvertRearChannelBleedPhase;
+
+	/* If true, send Surround Rear Channel Bleed Amount sends front left to back right and vice versa */
+	UPROPERTY(EditAnywhere, BlueprintReadWrite, Category = SubmixEffectPreset)
+	bool bSurroundRearChannelFlip;
+};
+
+
+/** Audio render thread effect object. */
+class SYNTHESIS_API FSubmixEffectConvolutionReverb : public FSoundEffectSubmix
+{
+public:
+	// Gain entry into convolution matrix
+	struct FConvolutionGainEntry
+	{
+		int32 InputIndex = 0;
+		int32 ImpulseIndex = 0;
+		int32 OutputIndex = 0;
+		float Gain = 0.f;
+
+		FConvolutionGainEntry(int32 InInputIndex, int32 InImpulseIndex, int32 InOutputIndex, float InGain)
+		:	InputIndex(InInputIndex)
+		,	ImpulseIndex(InImpulseIndex)
+		,	OutputIndex(InOutputIndex)
+		,	Gain(InGain)
+		{
+		}
+	};
+
+	// Data used to initialize the convolution algorithm
+	struct FConvolutionAlgorithmInitData
+	{
+		Audio::FConvolutionSettings AlgorithmSettings;
+		float ImpulseSampleRate = 0.f;
+		float TargetSampleRate = 0.f;
+		TArray<FConvolutionGainEntry> Gains;
+		TArray<float> Samples;
+	};
+
+	typedef int32 FConvolutionAlgorithmID;
+
+	// Data used to track versions of current running convolution.
+	struct FVersionData
+	{
+		FConvolutionAlgorithmID ConvolutionID;
+
+		FVersionData();
+		bool operator==(const FVersionData& Other) const;
+	};
+
+	
+	FSubmixEffectConvolutionReverb() = delete;
+	// Construct a convolution object with an existing preset. 
+	FSubmixEffectConvolutionReverb(const USubmixEffectConvolutionReverbPreset* InPreset);
+
+	~FSubmixEffectConvolutionReverb();
+
+	// Called on an audio effect at initialization on main thread before audio processing begins.
+	virtual void Init(const FSoundEffectSubmixInitData& InitData) override;
+
+	// Called when an audio effect preset settings is changed.
+	virtual void OnPresetChanged() override;
+
+	// Process the input block of audio. Called on audio thread.
+	virtual void OnProcessAudio(const FSoundEffectSubmixInputData& InData, FSoundEffectSubmixOutputData& OutData) override;
+
+	// We want to receive down-mixed submix audio to stereo input for the reverb effect
+	virtual uint32 GetDesiredInputChannelCountOverride() const override;
+
+	// Call on the game thread in order to update the impulse response and hardware acceleration
+	// used in this submix effect.
+	FVersionData UpdateConvolutionAlgorithm(const USubmixEffectConvolutionReverbPreset* InPreset);
+
+	// Sets the convolution algorithm object used if the version data matches the most up-to-date
+	// version data tracked within this object. This method must be called on the audio render thread.
+	void SetConvolutionAlgorithmIfCurrent(TUniquePtr<Audio::IConvolutionAlgorithm> InAlgo, const FVersionData& InVersionData);
+
+	// Returns the data needed to build a convolution algorithm.
+	FConvolutionAlgorithmInitData GetConvolutionAlgorithmInitData() const;
+
+private:
+
+	// Increments the internal verion number and returns a copy of the
+	// latest version information.
+	FVersionData UpdateVersion();
+
+	// Updates parameters on the render thread.
+	void UpdateParameters();
+
+	// Params object for copying between audio render thread and outside threads.
+	Audio::TParams<Audio::FConvolutionReverbSettings> Params;
+
+	// ConvolutionReverb performs majority of DSP operations
+	Audio::FConvolutionReverb ConvolutionReverb;
+
+	// Audio sample rate
+	float SampleRate;
+
+	// Number of input and output channels. These are TAtomic
+	// since they can be accessed on multiple threads.
+	TAtomic<int32> NumInputChannels;
+	TAtomic<int32> NumOutputChannels;
+
+	mutable FCriticalSection VersionDataCriticalSection;
+	FVersionData VersionData;
+
+	struct FIRAssetData 
+	{
+		TArray<float> Samples;
+		int32 NumChannels;
+		float SampleRate;
+		int32 BlockSize;
+		bool bEnableHardwareAcceleration;
+
+		FIRAssetData();
+	};
+
+	mutable FCriticalSection IRAssetCriticalSection;
+	// Internal copy of impulse response asset data.
+	FIRAssetData IRAssetData;
+};
+
+UCLASS()
+class SYNTHESIS_API USubmixEffectConvolutionReverbPreset : public USoundEffectSubmixPreset
+{
+	GENERATED_BODY()
+
+public:
+	
+	USubmixEffectConvolutionReverbPreset(const FObjectInitializer& ObjectInitializer);
+
+	/* Note: The SubmixEffect boilerplate macros could not be utilized here because
+	 * the "CreateNewEffect" implementation differed from those available in the
+	 * boilerplate macro.
+	 */
+
+	virtual bool CanFilter() const override;
+
+	virtual bool HasAssetActions() const;
+
+	virtual FText GetAssetActionName() const override;
+
+	virtual UClass* GetSupportedClass() const override;
+
+	virtual FSoundEffectBase* CreateNewEffect() const override;
+
+	virtual USoundEffectPreset* CreateNewPreset(UObject* InParent, FName Name, EObjectFlags Flags) const override;
+
+	virtual void Init() override;
+
+
+	FSubmixEffectConvolutionReverbSettings GetSettings() const;
+
+	/** Set the convolution reverb settings */
+	UFUNCTION(BlueprintCallable, Category = "Audio|Effects")
+	void SetSettings(const FSubmixEffectConvolutionReverbSettings& InSettings);
+
+	/** Set the convolution reverb impulse response */
+	UFUNCTION(BlueprintCallable, BlueprintSetter, Category = "Audio|Effects")
+	void SetImpulseResponse(UAudioImpulseResponse* InImpulseResponse);
+
+	/** ConvolutionReverbPreset Preset Settings. */
+	UPROPERTY(EditAnywhere, BlueprintReadWrite, BlueprintSetter = SetSettings, Category = SubmixEffectPreset)
+	FSubmixEffectConvolutionReverbSettings Settings;
+
+	/** The impulse response used for convolution. */
+	UPROPERTY(EditAnywhere, BlueprintReadWrite, BlueprintSetter = SetImpulseResponse, Category = SubmixEffectPreset)
+	UAudioImpulseResponse* ImpulseResponse;
+
+	/** Set the internal block size. This can effect latency and performance. Higher values will result in
+	 * lower CPU costs while lower values will result higher CPU costs. Latency may be affected depending
+	 * on the interplay between audio engines buffer sizes and this effects block size. Generally, higher
+	 * values result in higher latency, and lower values result in lower latency. 
+	 */
+	UPROPERTY(EditAnywhere, BlueprintReadOnly, AdvancedDisplay, Category = SubmixEffectPreset)
+	ESubmixEffectConvolutionReverbBlockSize BlockSize;
+
+	/** Opt into hardware acceleration of the convolution reverb (if available) */
+	UPROPERTY(EditAnywhere, BlueprintReadOnly, AdvancedDisplay, Category = SubmixEffectPreset)
+	bool bEnableHardwareAcceleration;
+	
+
+#if WITH_EDITORONLY_DATA
+	// Binds to the UAudioImpulseRespont::OnObjectPropertyChanged delegate of the current ImpulseResponse
+	void BindToImpulseResponseObjectChange();
+
+	virtual void PreEditChange(FProperty* PropertyAboutToChange) override;
+	virtual void PostEditChangeProperty(FPropertyChangedEvent& PropertyChangedEvent) override;
+	virtual void PostLoad() override;
+
+	// Called when a property changes on teh ImpulseResponse object
+	void PostEditChangeImpulseProperty(FPropertyChangedEvent& PropertyChangedEvent);
+#endif
+
+private:
+
+	void UpdateSettings();
+
+	// This method requires that the submix effect is registered with a preset.  If this 
+	// submix effect is not registered with a preset, then this will not update the convolution
+	// algorithm.
+	void UpdateEffectConvolutionAlgorithm();
+
+	mutable FCriticalSection SettingsCritSect; 
+	FSubmixEffectConvolutionReverbSettings SettingsCopy; 
+
+#if WITH_EDITORONLY_DATA
+	TMap<UObject*, FDelegateHandle> DelegateHandles;
+#endif
+};
+
diff --git a/Engine/Plugins/Runtime/Synthesis/Source/Synthesis/Private/ConvolutionAlgorithmWrapper.cpp b/Engine/Plugins/Runtime/Synthesis/Source/Synthesis/Private/ConvolutionAlgorithmWrapper.cpp
deleted file mode 100644
diff --git a/Engine/Plugins/Runtime/Synthesis/Source/Synthesis/Private/ConvolutionReverb.cpp b/Engine/Plugins/Runtime/Synthesis/Source/Synthesis/Private/ConvolutionReverb.cpp
new file mode 100644
--- /dev/null
+++ b/Engine/Plugins/Runtime/Synthesis/Source/Synthesis/Private/ConvolutionReverb.cpp
@@ -0,0 +1,371 @@
+// Copyright Epic Games, Inc. All Rights Reserved.
+
+#include "ConvolutionReverb.h"
+
+#include "CoreMinimal.h"
+#include "SynthesisModule.h"
+#include "DSP/ConvolutionAlgorithm.h"
+#include "DSP/BufferVectorOperations.h"
+
+namespace Audio
+{
+	const FConvolutionReverbSettings FConvolutionReverbSettings::DefaultSettings;
+	FConvolutionReverbSettings::FConvolutionReverbSettings()
+	:	NormalizationVolume(-24.f)
+	,	RearChannelBleed(0.f)
+	,	bRearChannelFlip(false)
+	{}
+
+
+	FConvolutionReverb::FConvolutionReverb(TUniquePtr<IConvolutionAlgorithm> InAlgorithm, const FConvolutionReverbSettings& InSettings)
+	:	Settings()
+	,	ConvolutionAlgorithm(nullptr)
+	,	OutputGain(1.f)
+	,	bIsSurroundOutput(false)
+	,	bHasSurroundBleed(false)
+	,	ExpectedNumFramesPerCallback(0)
+	,	NumInputSamplesPerBlock(0)
+	,	NumOutputSamplesPerBlock(0)
+	{
+		SetSettings(InSettings);
+		SetConvolutionAlgorithm(MoveTemp(InAlgorithm));
+		UpdateBlockBuffers();
+	}
+
+	void FConvolutionReverb::SetSettings(const FConvolutionReverbSettings& InSettings)
+	{
+		Settings = InSettings;
+
+		Update();
+	}
+
+	const FConvolutionReverbSettings& FConvolutionReverb::GetSettings() const
+	{
+		return Settings;
+	}
+
+	void FConvolutionReverb::SetConvolutionAlgorithm(TUniquePtr<IConvolutionAlgorithm> InAlgorithm)
+	{
+		ConvolutionAlgorithm = MoveTemp(InAlgorithm);
+
+		Update();
+	}
+
+	void FConvolutionReverb::ProcessAudio(int32 NumInputChannels, AlignedFloatBuffer& InputAudio, int32 NumOutputChannels, AlignedFloatBuffer& OutputAudio)
+	{
+		check(InputBlockBuffer.IsValid());
+		check(OutputBlockBuffer.IsValid());
+
+		OutputAudio.Reset();
+
+		const int32 NumInputSamples = InputAudio.Num();
+		const int32 NumFrames = NumInputSamples / NumInputChannels;
+		check((NumInputSamples % NumFrames) == 0);
+		const int32 NumOutputSamples = NumFrames * NumOutputChannels;
+
+		if (NumFrames != ExpectedNumFramesPerCallback)
+		{
+			ExpectedNumFramesPerCallback = NumFrames;
+			// If the number of frames processed per a call changes, then
+			// block buffers need to be updated. 
+			UpdateBlockBuffers();
+		}
+
+		if (NumFrames < 1)
+		{
+			return;
+		}
+
+		if (NumOutputSamples > 0)
+		{
+			OutputAudio.AddZeroed(NumOutputSamples);
+		}
+
+		// If we don't have a valid algorithm. We cannot do any processing. 
+		if (!ConvolutionAlgorithm.IsValid())
+		{
+			// Early exit due to invalid algorithm pointer
+			return;
+		}
+
+		// Our convolution algorithm should have equal number of input and output channels.
+		if (NumInputChannels != ConvolutionAlgorithm->GetNumAudioInputs())
+		{
+			UE_LOG(LogSynthesis, Error, TEXT("Convolution num input channel mismatch. Expected %d, got %d"), ConvolutionAlgorithm->GetNumAudioInputs(), NumInputChannels);
+			return;
+		}
+		else if (NumOutputChannels != ConvolutionAlgorithm->GetNumAudioOutputs())
+		{
+			UE_LOG(LogSynthesis, Error, TEXT("Convolution num output channel mismatch. Expected %d, got %d"), ConvolutionAlgorithm->GetNumAudioOutputs(), NumOutputChannels);
+			return;
+		}
+
+		// Process all available audio that completes a block.
+		InputBlockBuffer->AddSamples(InputAudio.GetData(), NumInputSamples);
+
+		while (InputBlockBuffer->GetNumAvailable() >= NumInputSamplesPerBlock)
+		{
+			const float* InputBlock = InputBlockBuffer->InspectSamples(NumInputSamplesPerBlock);
+
+			if (nullptr != InputBlock)
+			{
+				ProcessAudioBlock(InputBlock, NumInputChannels, InterleavedOutputBlock, NumOutputChannels);
+				OutputBlockBuffer->AddSamples(InterleavedOutputBlock.GetData(), InterleavedOutputBlock.Num());
+			}
+
+			InputBlockBuffer->RemoveSamples(NumInputSamplesPerBlock);
+		}
+
+		// Pop off all audio that's ready to go.
+		int32 NumToPop = FMath::Min(OutputBlockBuffer->GetNumAvailable(), NumOutputSamples);
+
+		if (NumToPop > 0)
+		{
+			int32 OutputCopyOffset = NumOutputSamples - NumToPop;
+			const float* OutputPtr = OutputBlockBuffer->InspectSamples(NumToPop);
+			if (nullptr != OutputPtr)
+			{
+				FMemory::Memcpy(&OutputAudio.GetData()[OutputCopyOffset], OutputPtr, NumToPop * sizeof(float));
+			}
+			OutputBlockBuffer->RemoveSamples(NumToPop);
+		}
+	}
+
+	void FConvolutionReverb::ProcessAudioBlock(const float* InputAudio, int32 InNumInputChannels, AlignedFloatBuffer& OutputAudio, int32 InNumOutputChannels)
+	{
+		check(nullptr != InputAudio);
+		check(ConvolutionAlgorithm.IsValid());
+
+		const int32 NumFrames = ConvolutionAlgorithm->GetNumSamplesInBlock();
+		const int32 NumInputSamples = InNumInputChannels * NumFrames;
+		const int32 NumOutputSamples = InNumOutputChannels * NumFrames;
+
+		TArrayView<const float> InputView(InputAudio, NumInputSamples);
+
+		// De-interleave Input buffer into scratch input buffer
+		DeinterleaveBuffer(InputDeinterleaveBuffers, InputView, InNumInputChannels);
+
+		// setup pointers for convolution in case any reallocation has happened.
+		for (int32 i = 0; i < InNumInputChannels; i++)
+		{
+			InputBufferPtrs[i] = InputDeinterleaveBuffers[i].GetData();
+		}
+
+		for (int32 i = 0; i < InNumOutputChannels; i++)
+		{
+			OutputBufferPtrs[i] = OutputDeinterleaveBuffers[i].GetData();
+		}
+
+		// get next buffer of convolution output (store in scratch output buffer)
+		ConvolutionAlgorithm->ProcessAudioBlock(InputBufferPtrs.GetData(), OutputBufferPtrs.GetData());
+
+		// add surround bleed
+		if (bIsSurroundOutput && bHasSurroundBleed)
+		{
+			const int32 OutputChannelOffset = InNumOutputChannels - 2;
+
+			if (!Settings.bRearChannelFlip)
+			{
+				MixInBufferFast(OutputDeinterleaveBuffers[0], OutputDeinterleaveBuffers[OutputChannelOffset], Settings.RearChannelBleed);
+				MixInBufferFast(OutputDeinterleaveBuffers[1], OutputDeinterleaveBuffers[OutputChannelOffset + 1], Settings.RearChannelBleed);
+			}
+			else
+			{
+				MixInBufferFast(OutputDeinterleaveBuffers[1], OutputDeinterleaveBuffers[OutputChannelOffset], Settings.RearChannelBleed);
+				MixInBufferFast(OutputDeinterleaveBuffers[0], OutputDeinterleaveBuffers[OutputChannelOffset + 1], Settings.RearChannelBleed);
+			}
+		}
+			
+		// re-interleave scratch output buffer into final output buffer
+		InterleaveBuffer(OutputAudio, OutputDeinterleaveBuffers, InNumOutputChannels);
+
+		// Apply final gain
+		Audio::MultiplyBufferByConstantInPlace(OutputAudio, OutputGain);
+	}
+
+	int32 FConvolutionReverb::GetNumInputChannels() const
+	{
+		if (!ConvolutionAlgorithm.IsValid())
+		{
+			return 0;
+		}
+		return ConvolutionAlgorithm->GetNumAudioInputs();
+	}
+
+	int32 FConvolutionReverb::GetNumOutputChannels() const
+	{
+		if (!ConvolutionAlgorithm.IsValid())
+		{
+			return 0;
+		}
+		return ConvolutionAlgorithm->GetNumAudioOutputs();
+	}
+
+	void FConvolutionReverb::Update()
+	{
+		OutputGain = Settings.NormalizationVolume;
+
+		int32 NumOutputChannels = 0;
+		int32 NumInputChannels = 0;
+		int32 NumFrames = 0;
+
+		NumInputSamplesPerBlock = 0;
+		NumOutputSamplesPerBlock = 0;
+
+		if (ConvolutionAlgorithm.IsValid())
+		{
+			NumOutputChannels = ConvolutionAlgorithm->GetNumAudioOutputs();
+			NumInputChannels = ConvolutionAlgorithm->GetNumAudioInputs();
+			NumFrames = ConvolutionAlgorithm->GetNumSamplesInBlock();
+
+			NumInputSamplesPerBlock = NumInputChannels * NumFrames;
+			NumOutputSamplesPerBlock = NumOutputChannels * NumFrames;
+
+			while (InputDeinterleaveBuffers.Num() < NumInputChannels)
+			{
+				InputDeinterleaveBuffers.Emplace();
+			}
+
+			while (OutputDeinterleaveBuffers.Num() < NumOutputChannels)
+			{
+				OutputDeinterleaveBuffers.Emplace();
+			}
+
+			InputBufferPtrs.SetNumZeroed(NumInputChannels);
+			OutputBufferPtrs.SetNumZeroed(NumOutputChannels);
+
+			for (int32 i = 0; i < NumInputChannels; i++)
+			{
+				AlignedFloatBuffer& Buffer = InputDeinterleaveBuffers[i];
+				Buffer.Reset();
+				if (NumFrames > 0)
+				{
+					Buffer.AddUninitialized(NumFrames);
+				}
+				InputBufferPtrs[i] = Buffer.GetData();
+			}
+
+			for (int32 i = 0; i < NumOutputChannels; i++)
+			{
+				AlignedFloatBuffer& Buffer = OutputDeinterleaveBuffers[i];
+				Buffer.Reset();
+				if (NumFrames > 0)
+				{
+					Buffer.AddUninitialized(NumFrames);
+				}
+				OutputBufferPtrs[i] = Buffer.GetData();
+			}
+		}
+
+		bIsSurroundOutput = (NumOutputChannels >= 4);
+		bHasSurroundBleed = FMath::Abs(Settings.RearChannelBleed) > 0.00101f; // > -59.9db
+		
+		// Modify the output level if we are using surround bleed.
+		// The gain scalar is applied to compensate for louder output
+		// if our surround output signal is folded down to stereo
+		if (bIsSurroundOutput && bHasSurroundBleed && Settings.bRearChannelFlip)
+		{
+			OutputGain *= FMath::Sqrt(1.0f/2.0f);
+		}
+		else if (bIsSurroundOutput && bHasSurroundBleed)
+		{
+			OutputGain *= 0.5f;
+		}
+
+		UpdateBlockBuffers();
+	}
+
+	void FConvolutionReverb::UpdateBlockBuffers()
+	{
+		int32 AlgoBlockSize = 0;
+		int32 NumInputChannels = 0;
+		int32 NumOutputChannels = 0;
+
+		if (ConvolutionAlgorithm.IsValid())
+		{
+			NumInputChannels = ConvolutionAlgorithm->GetNumAudioInputs();
+			NumOutputChannels = ConvolutionAlgorithm->GetNumAudioOutputs();
+			AlgoBlockSize = ConvolutionAlgorithm->GetNumSamplesInBlock();
+		}
+
+		int32 ChunkSize = FMath::Max(ExpectedNumFramesPerCallback, AlgoBlockSize);
+		ChunkSize = FMath::Max(128, ChunkSize);
+
+		int32 InputCapacity = FMath::Max(8192, 2 * ChunkSize * NumInputChannels);
+		int32 InputMaxInspect = FMath::Max(1024, AlgoBlockSize * NumInputChannels);
+
+		InputBlockBuffer = MakeUnique<FAlignedBlockBuffer>(InputCapacity, InputMaxInspect);
+
+		int32 OutputCapacity = FMath::Max(8192, 2 * ChunkSize * NumOutputChannels);
+		int32 OutputMaxInspect = FMath::Max(1024, ExpectedNumFramesPerCallback * NumOutputChannels);
+
+		OutputBlockBuffer = MakeUnique<FAlignedBlockBuffer>(OutputCapacity, OutputMaxInspect);
+	}
+
+	void FConvolutionReverb::InterleaveBuffer(AlignedFloatBuffer& OutBuffer, const TArray<AlignedFloatBuffer>& InputBuffers, const int32 NumChannels)
+	{
+		check(InputBuffers.Num() >= NumChannels);
+
+		OutBuffer.Reset();
+
+		if (InputBuffers.Num() < 1)
+		{
+			return;
+		}
+
+		int32 NumFrames = InputBuffers[0].Num();
+		int32 NumOutputSamples = NumFrames * NumChannels;
+
+		if (NumOutputSamples > 0)
+		{
+			OutBuffer.AddUninitialized(NumOutputSamples);
+		}
+
+		float* OutData = OutBuffer.GetData();
+		for (int32 i = 0; i < NumChannels; ++i)
+		{
+			check(InputBuffers[i].Num() == NumFrames);
+
+			const float* InData = InputBuffers[i].GetData();
+
+			int32 OutPos = i;
+
+			for (int32 j = 0; j < NumFrames; j++)
+			{
+				OutData[OutPos] = InData[j];
+				OutPos += NumChannels;
+			}
+		}
+	}
+
+	void FConvolutionReverb::DeinterleaveBuffer(TArray<AlignedFloatBuffer>& OutputBuffers, TArrayView<const float> InputBuffer, const int32 NumChannels)
+	{
+		check(OutputBuffers.Num() >= NumChannels);
+
+		int32 NumFrames = InputBuffer.Num();
+
+		if (NumChannels > 0)
+		{
+			NumFrames /= NumChannels;
+		}
+
+		const float* InputData = InputBuffer.GetData();
+
+		for (int32 i = 0; i < NumChannels; i++)
+		{
+			AlignedFloatBuffer& OutBuffer = OutputBuffers[i];
+
+			check(OutBuffer.Num() == NumFrames);
+
+			float* OutData = OutBuffer.GetData();
+
+			int32 InputPos = i;
+
+			for (int32 j = 0; j < NumFrames; j++)
+			{
+				OutData[j] = InputData[InputPos];
+				InputPos += NumChannels;
+			}
+		}
+	}
+}
diff --git a/Engine/Plugins/Runtime/Synthesis/Source/Synthesis/Private/ConvolutionReverbComponent.cpp b/Engine/Plugins/Runtime/Synthesis/Source/Synthesis/Private/ConvolutionReverbComponent.cpp
deleted file mode 100644
diff --git a/Engine/Plugins/Runtime/Synthesis/Source/Synthesis/Private/ImpulseResponseAsset.cpp b/Engine/Plugins/Runtime/Synthesis/Source/Synthesis/Private/ImpulseResponseAsset.cpp
deleted file mode 100644
diff --git a/Engine/Plugins/Runtime/Synthesis/Source/Synthesis/Private/SubmixEffectConvolutionReverb.cpp b/Engine/Plugins/Runtime/Synthesis/Source/Synthesis/Private/SubmixEffectConvolutionReverb.cpp
new file mode 100644
--- /dev/null
+++ b/Engine/Plugins/Runtime/Synthesis/Source/Synthesis/Private/SubmixEffectConvolutionReverb.cpp
@@ -0,0 +1,728 @@
+// Copyright Epic Games, Inc. All Rights Reserved.
+
+#include "SubmixEffects/SubmixEffectConvolutionReverb.h"
+
+#include "CoreMinimal.h"
+#include "Async/Async.h"
+#include "ConvolutionReverb.h"
+#include "DSP/ConvolutionAlgorithm.h"
+#include "DSP/SampleRateConverter.h"
+#include "SynthesisModule.h"
+
+namespace AudioConvReverbIntrinsics
+{
+	// names of various smart pointers
+	using FSubmixEffectConvSharedPtr = TSharedPtr<FSubmixEffectConvolutionReverb, ESPMode::ThreadSafe>;
+	using FSubmixEffectConvWeakPtr = TWeakPtr<FSubmixEffectConvolutionReverb, ESPMode::ThreadSafe>;
+	using FConvAlgoUniquePtr = TUniquePtr<Audio::IConvolutionAlgorithm>;
+
+	// names of nested classes and structs
+	using FConvolutionAlgorithmInitData = FSubmixEffectConvolutionReverb::FConvolutionAlgorithmInitData;
+	using FConvolutionGainEntry = FSubmixEffectConvolutionReverb::FConvolutionGainEntry;
+	using FVersionData = FSubmixEffectConvolutionReverb::FVersionData;
+
+
+	// Create a convolution algorithm. This performs creation of the convolution algorithm object,
+	// converting sample rates of impulse responses, sets the impulse response and initializes the
+	// gain matrix of the convolution algorithm.
+	//
+	// @params InInitData - Contains all the information needed to create a convolution algorithm.
+	//
+	// @return TUniquePtr<Audio::IConvolutionAlgorithm>  Will be invalid if there was an error.
+	static FConvAlgoUniquePtr CreateConvolutionAlgorithm(const FConvolutionAlgorithmInitData& InInitData)
+	{
+		// Check valid sample rates
+		if ((InInitData.ImpulseSampleRate <= 0.f) || (InInitData.TargetSampleRate <= 0.f))
+		{
+			return FConvAlgoUniquePtr();
+		}
+
+		const Audio::FConvolutionSettings& AlgoSettings = InInitData.AlgorithmSettings;
+
+		// Check valid channel counts
+		if ((AlgoSettings.NumInputChannels < 1) || (AlgoSettings.NumOutputChannels < 1) || (AlgoSettings.NumImpulseResponses < 1))
+		{
+			return FConvAlgoUniquePtr();
+		}
+
+		// Create convolution algorithm
+		FConvAlgoUniquePtr ConvolutionAlgorithm = Audio::FConvolutionFactory::NewConvolutionAlgorithm(AlgoSettings);
+
+		if (!ConvolutionAlgorithm.IsValid())
+		{
+			UE_LOG(LogSynthesis, Warning, TEXT("Failed to greate convolution algorithm for convolution reverb"));
+			return ConvolutionAlgorithm;
+		}
+
+		int32 NumFrames = InInitData.Samples.Num() / AlgoSettings.NumImpulseResponses;
+
+		// Prepare deinterleave pointers
+		TArray<Audio::AlignedFloatBuffer> DeinterleaveSamples;
+		while (DeinterleaveSamples.Num() < AlgoSettings.NumImpulseResponses)
+		{
+			Audio::AlignedFloatBuffer& Buffer = DeinterleaveSamples.Emplace_GetRef();
+			if (NumFrames > 0)
+			{
+				Buffer.AddUninitialized(NumFrames);
+			}
+		}
+
+		// Prepare impulse samples by converting samplerate and deinterleaving.
+		if (InInitData.ImpulseSampleRate == InInitData.TargetSampleRate)
+		{
+			// Deinterleave impulse samples
+			Audio::FConvolutionReverb::DeinterleaveBuffer(DeinterleaveSamples, InInitData.Samples, AlgoSettings.NumImpulseResponses);
+		}
+		else	
+		{
+			// convert sample rate of impulse 
+			TArray<float> TargetImpulseSamples;
+			float SampleRateRatio = InInitData.TargetSampleRate / InInitData.ImpulseSampleRate;
+
+			TUniquePtr<Audio::ISampleRateConverter> Converter(Audio::ISampleRateConverter::CreateSampleRateConverter());
+
+			if (!Converter.IsValid())
+			{
+				UE_LOG(LogSynthesis, Error, TEXT("Audio::ISampleRateConverter failed to create a sample rate converter"));
+				return FConvAlgoUniquePtr();
+			}
+
+			Converter->Init(SampleRateRatio, AlgoSettings.NumImpulseResponses);
+			Converter->ProcessFullbuffer(InInitData.Samples.GetData(), InInitData.Samples.Num(), TargetImpulseSamples);
+
+			// Deinterleave impulse samples
+			Audio::FConvolutionReverb::DeinterleaveBuffer(DeinterleaveSamples, TargetImpulseSamples, AlgoSettings.NumImpulseResponses);
+		}
+
+		// Set impulse responses in algorithm
+		for (int32 i = 0; i < DeinterleaveSamples.Num(); i++)
+		{
+			const Audio::AlignedFloatBuffer& Buffer = DeinterleaveSamples[i];
+
+			ConvolutionAlgorithm->SetImpulseResponse(i, Buffer.GetData(), Buffer.Num());
+		}
+
+		// Setup gain matrix in algorithm
+		for (const FConvolutionGainEntry& Entry : InInitData.Gains)
+		{
+			// If an entry exceeds the index, ignore. Will not log a warning
+			// or error as this might be a common occurence when channel counts
+			// change.
+			if (Entry.InputIndex >= ConvolutionAlgorithm->GetNumAudioInputs())
+			{
+				continue;
+			}
+
+			if (Entry.ImpulseIndex >= ConvolutionAlgorithm->GetNumImpulseResponses())
+			{
+				continue;
+			}
+
+			if (Entry.OutputIndex >= ConvolutionAlgorithm->GetNumAudioOutputs())
+			{
+				continue;
+			}
+
+			ConvolutionAlgorithm->SetMatrixGain(Entry.InputIndex, Entry.ImpulseIndex, Entry.OutputIndex, Entry.Gain);
+		}
+		
+		return ConvolutionAlgorithm;
+	}
+
+	// Task for creating convolution algorithm object.
+	class FCreateConvolutionAlgorithmTask : public FNonAbandonableTask
+	{
+		// This task can delete itself
+		friend class FAutoDeleteAsyncTask<FCreateConvolutionAlgorithmTask>;
+
+		public:
+			FCreateConvolutionAlgorithmTask(
+					FSubmixEffectConvWeakPtr InEffectObject,
+					FConvolutionAlgorithmInitData&& InInitData,
+					const FVersionData& InVersionData)
+			:	EffectWeakPtr(InEffectObject)
+			,	InitData(MoveTemp(InInitData))
+			,	VersionData(InVersionData)
+			{}
+
+			void DoWork()
+			{
+				// Build the convolution algorithm object. 
+				FConvAlgoUniquePtr ConvolutionAlgorithm = CreateConvolutionAlgorithm(InitData);
+
+				FSubmixEffectConvSharedPtr EffectSharedPtr = EffectWeakPtr.Pin();
+
+				// Check that the submix effect still exists.
+				if (EffectSharedPtr.IsValid())
+				{
+					// If the effect ptr is still valid, add to it's command queue to set the convolution
+					// algorithm object in the audio render thread.
+
+					TUniqueFunction<void()> Command = [Algo = MoveTemp(ConvolutionAlgorithm), EffectSharedPtr, VersionData = VersionData] () mutable
+					{
+						EffectSharedPtr->SetConvolutionAlgorithmIfCurrent(MoveTemp(Algo), VersionData);
+					};
+					EffectSharedPtr->EffectCommand(MoveTemp(Command));
+				}
+			}
+
+			FORCEINLINE TStatId GetStatId() const { RETURN_QUICK_DECLARE_CYCLE_STAT(CreateConvolutionAlgorithmTask,     STATGROUP_ThreadPoolAsyncTasks); }
+
+		private:
+			FSubmixEffectConvWeakPtr EffectWeakPtr;
+			FConvolutionAlgorithmInitData InitData;
+			FVersionData VersionData;
+	};
+}
+
+UAudioImpulseResponse::UAudioImpulseResponse()
+:	NumChannels(0)
+,	SampleRate(0)
+,	NormalizationVolumeDb(-24.f)
+{
+}
+
+#if WITH_EDITORONLY_DATA
+void UAudioImpulseResponse::PostEditChangeProperty(FPropertyChangedEvent& PropertyChangedEvent) 
+{
+	Super::PostEditChangeProperty(PropertyChangedEvent);
+
+	OnObjectPropertyChanged.Broadcast(PropertyChangedEvent);
+	/*
+	UObject* Outer = GetOuter();
+	if (nullptr == Outer)
+	{
+		return;
+	}
+
+	if (Outer->GetClass()->IsChildOf(USubmixEffectConvolutionReverbPreset::StaticClass()))
+	{
+		USubmixEffectConvolutionReverbPreset* Preset = CastChecked<USubmixEffectConvolutionReverbPreset>(Outer);
+
+		if (nullptr == Preset)
+		{
+			return;
+		}
+
+		Preset->PostEditChangeImpulseProperty(PropertyChangedEvent);
+	}
+	*/
+}
+#endif
+
+FSubmixEffectConvolutionReverbSettings::FSubmixEffectConvolutionReverbSettings()
+:	NormalizationVolumeDb(-24.f)
+,	SurroundRearChannelBleedDb(-60.f)
+,	bInvertRearChannelBleedPhase(false)
+,	bSurroundRearChannelFlip(false)
+{
+}
+
+FSubmixEffectConvolutionReverb::FVersionData::FVersionData()
+: 	ConvolutionID(0)
+{
+}
+
+bool FSubmixEffectConvolutionReverb::FVersionData::operator==(const FVersionData& Other) const
+{
+	// Check if versions are equal
+	return ConvolutionID == Other.ConvolutionID;
+}
+
+FSubmixEffectConvolutionReverb::FIRAssetData::FIRAssetData()
+:	NumChannels(0)
+,	SampleRate(0.f)
+,	bEnableHardwareAcceleration(false)
+{
+}
+
+FSubmixEffectConvolutionReverb::FSubmixEffectConvolutionReverb(const USubmixEffectConvolutionReverbPreset* InPreset)
+: 	ConvolutionReverb(TUniquePtr<Audio::IConvolutionAlgorithm>(nullptr))
+,	SampleRate(0.0f)
+,	NumInputChannels(2)
+,	NumOutputChannels(2)
+{
+	UpdateConvolutionAlgorithm(InPreset);
+}
+
+FSubmixEffectConvolutionReverb::~FSubmixEffectConvolutionReverb()
+{
+}
+
+void FSubmixEffectConvolutionReverb::Init(const FSoundEffectSubmixInitData& InitData)
+{
+	using namespace AudioConvReverbIntrinsics;
+
+	SampleRate = InitData.SampleRate;
+
+	FVersionData UpdatedVersionData = UpdateVersion();
+
+	// Create the convolution algorithm init data
+	FConvolutionAlgorithmInitData ConvolutionInitData = GetConvolutionAlgorithmInitData();
+
+	// Create the convolution algorithm
+	FConvAlgoUniquePtr ConvolutionAlgo = CreateConvolutionAlgorithm(MoveTemp(ConvolutionInitData));
+
+	// Set the convolution algorithm, ensuring that it's the most up-to-date
+	SetConvolutionAlgorithmIfCurrent(MoveTemp(ConvolutionAlgo), UpdatedVersionData);
+}
+
+
+void FSubmixEffectConvolutionReverb::OnPresetChanged()
+{
+	USubmixEffectConvolutionReverbPreset* ConvolutionPreset = CastChecked<USubmixEffectConvolutionReverbPreset>(Preset);
+	FSubmixEffectConvolutionReverbSettings Settings = ConvolutionPreset->GetSettings();;
+
+	Audio::FConvolutionReverbSettings ReverbSettings;
+
+
+	float NewVolume = Audio::ConvertToLinear(Settings.NormalizationVolumeDb);
+	float NewRearChannelBleed = Audio::ConvertToLinear(Settings.SurroundRearChannelBleedDb);
+
+	if (Settings.bInvertRearChannelBleedPhase)
+	{
+		NewRearChannelBleed *= -1.f;
+	}
+
+	ReverbSettings.NormalizationVolume = NewVolume;
+	ReverbSettings.RearChannelBleed = NewRearChannelBleed;
+	ReverbSettings.bRearChannelFlip = Settings.bSurroundRearChannelFlip;
+
+	Params.SetParams(ReverbSettings);
+}
+
+
+FSubmixEffectConvolutionReverb::FConvolutionAlgorithmInitData FSubmixEffectConvolutionReverb::GetConvolutionAlgorithmInitData() const
+{
+	FScopeLock IRAssetLock(&IRAssetCriticalSection);
+
+	int32 ConvNumInputChannels = NumInputChannels.Load();
+	int32 ConvNumOutputChannels = NumOutputChannels.Load();
+
+	// Create convolution algorithm settings.
+	FConvolutionAlgorithmInitData CreateConvolutionSettings;
+
+	// IConvolutionAlgorithm settings
+	CreateConvolutionSettings.AlgorithmSettings.bEnableHardwareAcceleration = IRAssetData.bEnableHardwareAcceleration;
+	CreateConvolutionSettings.AlgorithmSettings.BlockNumSamples = 256;
+	CreateConvolutionSettings.AlgorithmSettings.NumInputChannels = ConvNumInputChannels;
+	CreateConvolutionSettings.AlgorithmSettings.NumOutputChannels = ConvNumOutputChannels;
+	CreateConvolutionSettings.AlgorithmSettings.NumImpulseResponses = IRAssetData.NumChannels;
+	CreateConvolutionSettings.AlgorithmSettings.MaxNumImpulseResponseSamples = 0;
+
+	if (IRAssetData.NumChannels > 0)
+	{
+		CreateConvolutionSettings.AlgorithmSettings.MaxNumImpulseResponseSamples = (IRAssetData.Samples.Num() / IRAssetData.NumChannels) + 1;
+	}
+
+	// Sample rate conversion settings
+	CreateConvolutionSettings.ImpulseSampleRate = IRAssetData.SampleRate;
+	CreateConvolutionSettings.TargetSampleRate = SampleRate;
+
+	// Setup gain matrix. Currently we are only setting up a 1-to-1 mapping. But this 
+	// could be enhanced by adding info to the IR asset or Preset to express different
+	// desired mappings.
+	int32 MinChannelCount = FMath::Min3(ConvNumInputChannels, ConvNumOutputChannels, IRAssetData.NumChannels);
+	for (int32 i = 0; i < MinChannelCount; i++)
+	{
+		CreateConvolutionSettings.Gains.Emplace(i, i, i, 1.f);
+	}
+
+	CreateConvolutionSettings.Samples = IRAssetData.Samples;
+
+	return CreateConvolutionSettings;
+}
+
+
+void FSubmixEffectConvolutionReverb::OnProcessAudio(const FSoundEffectSubmixInputData& InData, FSoundEffectSubmixOutputData& OutData)
+{
+	check(nullptr != InData.AudioBuffer);
+	check(nullptr != OutData.AudioBuffer);
+
+	UpdateParameters();
+
+	int32 ExpectedNumInputChannels = ConvolutionReverb.GetNumInputChannels();
+	int32 ExpectedNumOutputChannels = ConvolutionReverb.GetNumOutputChannels();
+
+	// Check if there is a channel mismatch between teh convolution reverb and the input/output data
+	bool bNumChannelsMismatch = false;
+	bNumChannelsMismatch |= (ExpectedNumInputChannels != InData.NumChannels);
+	bNumChannelsMismatch |= (ExpectedNumOutputChannels != OutData.NumChannels);
+
+	if (bNumChannelsMismatch)
+	{
+		int32 NumInputChannelsTemp = NumInputChannels.Load();
+		int32 NumOutputChannelsTemp = NumOutputChannels.Load();
+
+		// We should create a new algorithm if the input and output channels are non-zero.
+		// But we check the NumInputChannelsTemp/NumOutputChannelsTemp to check if we've
+		// already tried and failed to create an algorithm for this channel configuration.
+		bool bShouldCreateNewAlgo = true;
+		bShouldCreateNewAlgo &= (InData.NumChannels > 0);
+		bShouldCreateNewAlgo &= (OutData.NumChannels > 0);
+		bShouldCreateNewAlgo &= (NumInputChannelsTemp != InData.NumChannels);
+		bShouldCreateNewAlgo &= (NumOutputChannelsTemp != OutData.NumChannels);
+
+		if (bShouldCreateNewAlgo)
+		{
+			using namespace AudioConvReverbIntrinsics;
+
+			UE_LOG(LogSynthesis, Log, TEXT("Creating new convolution algorithm due to channel count update. Num Inputs %d -> %d. Num Outputs %d -> %d"), ExpectedNumInputChannels, InData.NumChannels, ExpectedNumOutputChannels, OutData.NumChannels);
+
+			// We don't update version data when changing the channel configuration
+			FVersionData CurrentVersionData;
+			{
+				FScopeLock VersionLock(&VersionDataCriticalSection);
+				CurrentVersionData = VersionData;
+			}
+
+			// These should only be updated here and in the constructor since we use them
+			// to see if we've tried this channel configuration already.
+			NumInputChannels.Store(InData.NumChannels);
+			NumOutputChannels.Store(OutData.NumChannels);
+
+			// I would really rather this be done in a task. To do so, we would need 
+			// access to a weak ptr to "this". Currently that is not possible since this object
+			// is owned by the Preset, and the TWeakObjectPtr cannot safely be copied across thread
+			// boundaries. 
+			{
+				// Create the convolution algorithm init data
+				FConvolutionAlgorithmInitData ConvolutionInitData = GetConvolutionAlgorithmInitData();
+
+				// Create the convolution algorithm
+				FConvAlgoUniquePtr ConvolutionAlgo = CreateConvolutionAlgorithm(MoveTemp(ConvolutionInitData));
+
+				// Set the convolution algorithm, ensuring that it's the most up-to-date
+				SetConvolutionAlgorithmIfCurrent(MoveTemp(ConvolutionAlgo), CurrentVersionData);
+			}
+
+			ExpectedNumInputChannels = ConvolutionReverb.GetNumInputChannels();
+			ExpectedNumOutputChannels = ConvolutionReverb.GetNumOutputChannels();
+		}
+	}
+
+
+	if ((ExpectedNumInputChannels == InData.NumChannels) && (ExpectedNumOutputChannels == OutData.NumChannels))
+	{
+		ConvolutionReverb.ProcessAudio(InData.NumChannels, *InData.AudioBuffer, OutData.NumChannels, *OutData.AudioBuffer);
+	}
+	else
+	{
+		// Channel mismatch. zero output data. Do *not* trigger rebuild here in case one is already in flight or simply because one cannot be built.
+		OutData.AudioBuffer->Reset();
+		int32 OutputNumFrames = InData.NumFrames * OutData.NumChannels;
+		if (OutputNumFrames > 0)
+		{
+			OutData.AudioBuffer->AddZeroed(OutputNumFrames);
+		}
+	}
+}
+
+FSubmixEffectConvolutionReverb::FVersionData FSubmixEffectConvolutionReverb::UpdateVersion()
+{
+	FVersionData VersionDataCopy;
+	FScopeLock VersionDataLock(&VersionDataCriticalSection);
+
+	VersionData.ConvolutionID++;
+	VersionDataCopy = VersionData;
+
+	return VersionDataCopy;
+}
+
+void FSubmixEffectConvolutionReverb::UpdateParameters()
+{
+	// Call on the audio render thread to update parameters.
+    Audio::FConvolutionReverbSettings NewSettings;
+    if (Params.GetParams(&NewSettings))
+    {
+        ConvolutionReverb.SetSettings(NewSettings);
+    }
+}
+
+
+uint32 FSubmixEffectConvolutionReverb::GetDesiredInputChannelCountOverride() const
+{
+	return 2;
+}
+
+FSubmixEffectConvolutionReverb::FVersionData FSubmixEffectConvolutionReverb::UpdateConvolutionAlgorithm(const USubmixEffectConvolutionReverbPreset* InPreset)
+{
+	// Copy data from preset into internal IRAssetData.
+	
+	check(IsInGameThread());
+
+	FScopeLock IRAssetLock(&IRAssetCriticalSection);
+
+	IRAssetData.Samples.Reset();
+	IRAssetData.NumChannels = 0;
+	IRAssetData.SampleRate = 0.f;
+	IRAssetData.bEnableHardwareAcceleration = false;
+	IRAssetData.BlockSize = 1024;
+
+	if (nullptr != InPreset)
+	{
+		if (InPreset->ImpulseResponse)
+		{
+			UAudioImpulseResponse* IR = InPreset->ImpulseResponse;
+
+			IRAssetData.Samples = IR->ImpulseResponse;
+			IRAssetData.NumChannels = IR->NumChannels;
+			IRAssetData.SampleRate = IR->SampleRate;
+		}
+		switch (InPreset->BlockSize)
+		{
+			case ESubmixEffectConvolutionReverbBlockSize::BlockSize256:
+				IRAssetData.BlockSize = 256;
+				break;
+
+			case ESubmixEffectConvolutionReverbBlockSize::BlockSize512:
+				IRAssetData.BlockSize = 512;
+				break;
+
+			case ESubmixEffectConvolutionReverbBlockSize::BlockSize1024:
+			default:
+				IRAssetData.BlockSize = 1024;
+				break;
+		}
+		IRAssetData.bEnableHardwareAcceleration = InPreset->bEnableHardwareAcceleration;
+	}
+
+	return UpdateVersion();
+}
+
+void FSubmixEffectConvolutionReverb::SetConvolutionAlgorithmIfCurrent(TUniquePtr<Audio::IConvolutionAlgorithm> InAlgo, const FVersionData& InVersionData)
+{
+	bool bIsCurrent = true;
+
+	{
+		FScopeLock VersionDataLock(&VersionDataCriticalSection);
+
+		bIsCurrent &= (InVersionData == VersionData);
+	}
+
+	if (bIsCurrent)
+	{
+		ConvolutionReverb.SetConvolutionAlgorithm(MoveTemp(InAlgo));
+	}
+}
+
+/********************************************************************/
+/***************** USubmixEffectConvolutionReverbPreset *************/
+/********************************************************************/
+
+USubmixEffectConvolutionReverbPreset::USubmixEffectConvolutionReverbPreset(const FObjectInitializer& ObjectInitializer)
+:	Super(ObjectInitializer)
+,	ImpulseResponse(nullptr)
+,	BlockSize(ESubmixEffectConvolutionReverbBlockSize::BlockSize1024)
+,	bEnableHardwareAcceleration(true)
+{}
+
+bool USubmixEffectConvolutionReverbPreset::CanFilter() const 
+{
+	return false; 
+}
+
+bool USubmixEffectConvolutionReverbPreset::HasAssetActions() const 
+{ 
+	return true; 
+}
+
+FText USubmixEffectConvolutionReverbPreset::GetAssetActionName() const 
+{ 
+	return FText::FromString(TEXT("SubmixEffectConvolutionReverb")); 
+}
+
+UClass* USubmixEffectConvolutionReverbPreset::GetSupportedClass() const 
+{ 
+	return USubmixEffectConvolutionReverbPreset::StaticClass(); 
+}
+
+FSoundEffectBase* USubmixEffectConvolutionReverbPreset::CreateNewEffect() const 
+{ 
+	// Pass a pointer to self into this constructor.  This enables the submix effect to have
+	// a copy of the impulse response data. The submix effect can then prepare the impulse
+	// response during it's "Init" routine.
+	return new FSubmixEffectConvolutionReverb(this); 
+}
+
+USoundEffectPreset* USubmixEffectConvolutionReverbPreset::CreateNewPreset(UObject* InParent, FName Name, EObjectFlags Flags) const
+{ 
+	USoundEffectPreset* NewPreset = NewObject<USubmixEffectConvolutionReverbPreset>(InParent, GetSupportedClass(), Name, Flags); 
+	NewPreset->Init(); 
+	return NewPreset; 
+}
+
+void USubmixEffectConvolutionReverbPreset::Init() 
+{ 
+	FScopeLock ScopeLock(&SettingsCritSect);
+	SettingsCopy = Settings;
+}
+
+void USubmixEffectConvolutionReverbPreset::UpdateSettings()
+{
+	{
+		// Copy settings to audio-render-thread version
+		FScopeLock ScopeLock(&SettingsCritSect);
+		SettingsCopy = Settings; 
+	}
+	// This sets the 'bChanged' to true on related effect instances which will 
+	// trigger a OnPresetChanged call on the audio render thread.
+	Update(); 
+} 
+
+FSubmixEffectConvolutionReverbSettings USubmixEffectConvolutionReverbPreset::GetSettings() const
+{ 
+	FScopeLock ScopeLock(&SettingsCritSect);
+	return SettingsCopy; 
+} 
+
+void USubmixEffectConvolutionReverbPreset::SetSettings(const FSubmixEffectConvolutionReverbSettings& InSettings)
+{
+	Settings = InSettings;
+	UpdateSettings();
+}
+
+void USubmixEffectConvolutionReverbPreset::SetImpulseResponse(UAudioImpulseResponse* InImpulseResponse)
+{
+	ImpulseResponse = InImpulseResponse;
+	UpdateEffectConvolutionAlgorithm();
+}
+
+void USubmixEffectConvolutionReverbPreset::UpdateEffectConvolutionAlgorithm()
+{
+	using namespace AudioConvReverbIntrinsics;
+	using FEffectWeakPtr = TWeakPtr<FSoundEffectBase, ESPMode::ThreadSafe>;
+	using FEffectSharedPtr = TSharedPtr<FSoundEffectBase, ESPMode::ThreadSafe>;
+
+	// Iterator over effects and set IR data.
+	for (FEffectWeakPtr& InstanceWeakPtr : Instances)
+	{
+		FEffectSharedPtr InstanceSharedPtr = InstanceWeakPtr.Pin();
+		if (InstanceSharedPtr.IsValid())
+		{
+			FSubmixEffectConvSharedPtr ConvEffectSharedPtr = StaticCastSharedPtr<FSubmixEffectConvolutionReverb>(InstanceSharedPtr);
+
+			FVersionData VersionData = ConvEffectSharedPtr->UpdateConvolutionAlgorithm(this);
+
+			FConvolutionAlgorithmInitData ConvolutionInitData = ConvEffectSharedPtr->GetConvolutionAlgorithmInitData();
+
+			(new FAutoDeleteAsyncTask<FCreateConvolutionAlgorithmTask>(ConvEffectSharedPtr, MoveTemp(ConvolutionInitData), VersionData))->StartBackgroundTask();
+
+		}
+	}
+}
+
+#if WITH_EDITORONLY_DATA
+void USubmixEffectConvolutionReverbPreset::BindToImpulseResponseObjectChange()
+{
+	if (nullptr != ImpulseResponse)
+	{
+		if (!DelegateHandles.Contains(ImpulseResponse))
+		{
+			FDelegateHandle Handle = ImpulseResponse->OnObjectPropertyChanged.AddUObject(this, &USubmixEffectConvolutionReverbPreset::PostEditChangeImpulseProperty);
+			if (Handle.IsValid())
+			{
+				DelegateHandles.Add(ImpulseResponse, Handle);
+			}
+		}
+	}
+}
+
+void USubmixEffectConvolutionReverbPreset::PreEditChange(FProperty* PropertyAboutToChange)
+{
+	Super::PreEditChange(PropertyAboutToChange);
+
+	if (nullptr == PropertyAboutToChange)
+	{
+		return;
+	}
+
+	const FName ImpulseResponseFName = GET_MEMBER_NAME_CHECKED(USubmixEffectConvolutionReverbPreset, ImpulseResponse);
+
+	if (PropertyAboutToChange->GetFName() == ImpulseResponseFName)
+	{
+		if (nullptr == ImpulseResponse)
+		{
+			return;
+		}
+
+		if (DelegateHandles.Contains(ImpulseResponse))
+		{
+			FDelegateHandle DelegateHandle = DelegateHandles[ImpulseResponse];
+
+			if (DelegateHandle.IsValid())
+			{
+				ImpulseResponse->OnObjectPropertyChanged.Remove(DelegateHandle);
+			}
+		}
+	}
+}
+
+void USubmixEffectConvolutionReverbPreset::PostLoad()
+{
+	Super::PostLoad();
+
+	BindToImpulseResponseObjectChange();
+}
+
+void USubmixEffectConvolutionReverbPreset::PostEditChangeProperty(FPropertyChangedEvent& PropertyChangedEvent)
+{
+	Super::PostEditChangeProperty(PropertyChangedEvent);
+
+	const FName ImpulseResponseFName = GET_MEMBER_NAME_CHECKED(USubmixEffectConvolutionReverbPreset, ImpulseResponse);
+	const TArray<FName> ConvolutionAlgorithmProperties(
+		{
+			ImpulseResponseFName,
+			GET_MEMBER_NAME_CHECKED(USubmixEffectConvolutionReverbPreset, bEnableHardwareAcceleration),
+			GET_MEMBER_NAME_CHECKED(USubmixEffectConvolutionReverbPreset, BlockSize)
+		}
+	);
+
+	if (!HasAnyFlags(RF_ClassDefaultObject))
+	{
+		if (FProperty* PropertyThatChanged = PropertyChangedEvent.Property)
+		{
+			const FName& Name = PropertyThatChanged->GetFName();
+
+			if (ConvolutionAlgorithmProperties.Contains(Name))
+			{
+				UpdateEffectConvolutionAlgorithm();
+			}
+
+			if (Name == ImpulseResponseFName)
+			{
+				BindToImpulseResponseObjectChange();
+
+				if (nullptr != ImpulseResponse)
+				{
+					Settings.NormalizationVolumeDb = ImpulseResponse->NormalizationVolumeDb;
+					UpdateSettings();
+				}
+			}
+		}
+	}
+}
+
+void USubmixEffectConvolutionReverbPreset::PostEditChangeImpulseProperty(FPropertyChangedEvent& PropertyChangedEvent) 
+{
+	const FName NormalizationVolumeFName = GET_MEMBER_NAME_CHECKED(UAudioImpulseResponse, NormalizationVolumeDb);
+
+	if (!HasAnyFlags(RF_ClassDefaultObject))
+	{
+		if (FProperty* PropertyThatChanged = PropertyChangedEvent.Property)
+		{
+			const FName& Name = PropertyThatChanged->GetFName();
+
+			if (Name == NormalizationVolumeFName)
+			{
+				if (nullptr != ImpulseResponse)
+				{
+					Settings.NormalizationVolumeDb = ImpulseResponse->NormalizationVolumeDb;
+					UpdateSettings();
+				}
+			}
+		}
+	}
+}
+
+#endif
diff --git a/Engine/Plugins/Runtime/Synthesis/Source/Synthesis/Private/SynthesisModule.cpp b/Engine/Plugins/Runtime/Synthesis/Source/Synthesis/Private/SynthesisModule.cpp
--- a/Engine/Plugins/Runtime/Synthesis/Source/Synthesis/Private/SynthesisModule.cpp
+++ b/Engine/Plugins/Runtime/Synthesis/Source/Synthesis/Private/SynthesisModule.cpp
@@ -14,6 +14,8 @@ IMPLEMENT_MODULE(FSynthesisModule, Synthesis)
 void FSynthesisModule::StartupModule()
 {
 	FSynthSlateStyleSet::Initialize();
+
+	FModuleManager::Get().LoadModuleChecked(TEXT("SignalProcessing"));
 }
 
 void FSynthesisModule::ShutdownModule()
diff --git a/Engine/Plugins/Runtime/Synthesis/Source/Synthesis/Public/ConvolutionAlgorithmWrapper.h b/Engine/Plugins/Runtime/Synthesis/Source/Synthesis/Public/ConvolutionAlgorithmWrapper.h
deleted file mode 100644
diff --git a/Engine/Plugins/Runtime/Synthesis/Source/Synthesis/Public/ConvolutionReverb.h b/Engine/Plugins/Runtime/Synthesis/Source/Synthesis/Public/ConvolutionReverb.h
new file mode 100644
--- /dev/null
+++ b/Engine/Plugins/Runtime/Synthesis/Source/Synthesis/Public/ConvolutionReverb.h
@@ -0,0 +1,83 @@
+// Copyright Epic Games, Inc. All Rights Reserved.
+
+#pragma once
+
+#include "CoreMinimal.h"
+#include "DSP/ConvolutionAlgorithm.h"
+#include "DSP/BufferVectorOperations.h"
+#include "DSP/AlignedBlockBuffer.h"
+
+namespace Audio
+{
+	struct SYNTHESIS_API FConvolutionReverbSettings
+	{
+		FConvolutionReverbSettings();
+
+		/* Used to account for energy added by convolution with "loud" Impulse Responses.  Not meant to be updated dynamically (linear gain)*/
+		float NormalizationVolume;
+
+		/* Amout of audio to be sent to rear channels in quad/surround configurations (linear gain, < 0 = phase inverted) */
+		float RearChannelBleed;
+
+		/* If true, send Surround Rear Channel Bleed Amount sends front left to back right and vice versa */
+		bool bRearChannelFlip;
+
+		static const FConvolutionReverbSettings DefaultSettings;
+	};
+
+	class SYNTHESIS_API FConvolutionReverb 
+	{
+		FConvolutionReverb() = delete;
+		FConvolutionReverb(const FConvolutionReverb&) = delete;
+		FConvolutionReverb& operator=(const FConvolutionReverb&) = delete;
+		FConvolutionReverb(const FConvolutionReverb&&) = delete;
+
+	public:
+		FConvolutionReverb(TUniquePtr<IConvolutionAlgorithm> InAlgorithm, const FConvolutionReverbSettings& InSettings = FConvolutionReverbSettings::DefaultSettings);
+		
+		void SetSettings(const FConvolutionReverbSettings& InSettings);
+
+		const FConvolutionReverbSettings& GetSettings() const;
+
+		void SetConvolutionAlgorithm(TUniquePtr<IConvolutionAlgorithm> InAlgorithm);
+
+		// If the number of input frames changes between callbacks, the output may contain discontinuities.
+		void ProcessAudio(int32 InNumInputChannels, AlignedFloatBuffer& InputAudio, int32 InNumOutputChannels, AlignedFloatBuffer& OutputAudio);
+
+		int32 GetNumInputChannels() const;
+		int32 GetNumOutputChannels() const;
+
+		static void InterleaveBuffer(AlignedFloatBuffer& OutBuffer, const TArray<AlignedFloatBuffer>& InputBuffers, const int32 NumChannels);
+		static void DeinterleaveBuffer(TArray<AlignedFloatBuffer>& OutputBuffers, TArrayView<const float> InputBuffer, const int32 NumChannels);
+
+	private:
+		void UpdateBlockBuffers();
+
+		void ProcessAudioBlock(const float* InputAudio, int32 InNumInputChannels, AlignedFloatBuffer& OutputAudio, int32 InNumOutputChannels);
+
+		void Update();
+
+		FConvolutionReverbSettings Settings;
+		TUniquePtr<IConvolutionAlgorithm> ConvolutionAlgorithm;
+		
+		// data is passed to the convolution algorithm as 2D arrays
+		TArray<AlignedFloatBuffer> InputDeinterleaveBuffers;
+
+		// data is recieved from the convolution algorithm as 2D arrays
+		TArray<AlignedFloatBuffer> OutputDeinterleaveBuffers;
+
+		TArray<float*> InputBufferPtrs;
+		TArray<float*> OutputBufferPtrs;
+
+		TUniquePtr<FAlignedBlockBuffer> InputBlockBuffer;
+		TUniquePtr<FAlignedBlockBuffer> OutputBlockBuffer;
+		AlignedFloatBuffer InterleavedOutputBlock;
+
+		float OutputGain;
+		bool bIsSurroundOutput;
+		bool bHasSurroundBleed;
+		int32 ExpectedNumFramesPerCallback;
+		int32 NumInputSamplesPerBlock;
+		int32 NumOutputSamplesPerBlock;
+	};
+}
diff --git a/Engine/Plugins/Runtime/Synthesis/Source/Synthesis/Public/ConvolutionReverbComponent.h b/Engine/Plugins/Runtime/Synthesis/Source/Synthesis/Public/ConvolutionReverbComponent.h
deleted file mode 100644
diff --git a/Engine/Plugins/Runtime/Synthesis/Source/Synthesis/Public/ImpulseResponseAsset.h b/Engine/Plugins/Runtime/Synthesis/Source/Synthesis/Public/ImpulseResponseAsset.h
deleted file mode 100644
diff --git a/Engine/Plugins/Runtime/Synthesis/Source/SynthesisEditor/Classes/AudioImpulseResponseAsset.h b/Engine/Plugins/Runtime/Synthesis/Source/SynthesisEditor/Classes/AudioImpulseResponseAsset.h
new file mode 100644
--- /dev/null
+++ b/Engine/Plugins/Runtime/Synthesis/Source/SynthesisEditor/Classes/AudioImpulseResponseAsset.h
@@ -0,0 +1,46 @@
+// Copyright Epic Games, Inc. All Rights Reserved.
+
+#pragma once
+
+#include "CoreMinimal.h"
+#include "Factories/Factory.h"
+#include "AssetTypeCategories.h"
+#include "AssetTypeActions_Base.h"
+#include "Sound/SoundWave.h"
+#include "Sound/SampleBufferIO.h"
+#include "AudioImpulseResponseAsset.generated.h"
+
+
+class FAssetTypeActions_AudioImpulseResponse : public FAssetTypeActions_Base
+{
+public:
+	virtual FText GetName() const override { return NSLOCTEXT("AssetTypeActions", "AssetTypeActions_AudioImpulseResponse", "Audio Impulse Respons"); }
+	virtual FColor GetTypeColor() const override { return FColor(0, 255, 0); }
+	virtual UClass* GetSupportedClass() const override;
+	virtual const TArray<FText>& GetSubMenus() const override;
+	virtual uint32 GetCategories() override { return EAssetTypeCategories::Sounds; }
+};
+
+class FAudioImpulseResponseExtension
+{
+public:
+	static void RegisterMenus();
+	static void GetExtendedActions(const struct FToolMenuContext& MenuContext);
+	static void ExecuteCreateImpulseResponse(const struct FToolMenuContext& MenuContext);
+};
+
+
+UCLASS(MinimalAPI, hidecategories = Object)
+class UAudioImpulseResponseFactory : public UFactory
+{
+	GENERATED_UCLASS_BODY()
+
+	//~ Begin UFactory Interface
+	virtual UObject* FactoryCreateNew(UClass* Class, UObject* InParent, FName Name, EObjectFlags Flags, UObject* Context, FFeedbackContext* Warn) override;
+	//~ Begin UFactory Interface
+	
+	TWeakObjectPtr<USoundWave> StagedSoundWave;
+
+private:
+	Audio::FSoundWavePCMLoader Loader;
+};
diff --git a/Engine/Plugins/Runtime/Synthesis/Source/SynthesisEditor/Classes/ImpulseResponse.h b/Engine/Plugins/Runtime/Synthesis/Source/SynthesisEditor/Classes/ImpulseResponse.h
deleted file mode 100644
diff --git a/Engine/Plugins/Runtime/Synthesis/Source/SynthesisEditor/Private/AudioImpulseResponseAsset.cpp b/Engine/Plugins/Runtime/Synthesis/Source/SynthesisEditor/Private/AudioImpulseResponseAsset.cpp
new file mode 100644
--- /dev/null
+++ b/Engine/Plugins/Runtime/Synthesis/Source/SynthesisEditor/Private/AudioImpulseResponseAsset.cpp
@@ -0,0 +1,134 @@
+// Copyright Epic Games, Inc. All Rights Reserved.
+#include "AudioImpulseResponseAsset.h"
+#include "UObject/ObjectMacros.h"
+#include "UObject/Object.h"
+#include "ToolMenus.h"
+#include "ContentBrowserMenuContexts.h"
+#include "Sound/SoundWave.h"
+#include "Sound/SampleBufferIO.h"
+#include "SampleBuffer.h"
+#include "SubmixEffects/SubmixEffectConvolutionReverb.h"
+
+#define LOCTEXT_NAMESPACE "AssetTypeActions"
+
+UClass* FAssetTypeActions_AudioImpulseResponse::GetSupportedClass() const
+{
+	return UAudioImpulseResponse::StaticClass();
+}
+
+const TArray<FText>& FAssetTypeActions_AudioImpulseResponse::GetSubMenus() const
+{
+	static const TArray<FText> SubMenus
+	{
+		FText(LOCTEXT("AssetConvolutionReverbSubmenu", "Convolution Reverb"))
+	};
+
+	return SubMenus;
+}
+
+void FAudioImpulseResponseExtension::RegisterMenus()
+{
+	if (!UToolMenus::IsToolMenuUIEnabled())
+	{
+		return;
+	}
+
+	FToolMenuOwnerScoped MenuOwner("ConvolutionReverb");
+	UToolMenu* Menu = UToolMenus::Get()->ExtendMenu("ContentBrowser.AssetContextMenu.SoundWave");
+	FToolMenuSection& Section = Menu->FindOrAddSection("GetAssetActions");
+
+	Section.AddDynamicEntry("SoundWaveAssetConversion", FNewToolMenuSectionDelegate::CreateLambda([](FToolMenuSection& InSection)
+	{
+		const TAttribute<FText> Label = LOCTEXT("SoundWave_CreateImpulseResponse", "Create Impulse Response");
+		const TAttribute<FText> ToolTip = LOCTEXT("SoundWave_CreateImpulseResponseTooltip", "Creates an impulse response asset using the selected sound wave.");
+		const FSlateIcon Icon = FSlateIcon(FEditorStyle::GetStyleSetName(), "ClassIcon.ImpulseResponse");
+		const FToolMenuExecuteAction UIAction = FToolMenuExecuteAction::CreateStatic(&FAudioImpulseResponseExtension::ExecuteCreateImpulseResponse);
+
+		InSection.AddMenuEntry("SoundWave_CreateImpulseResponse", Label, ToolTip, Icon, UIAction);
+	}));
+}
+
+void FAudioImpulseResponseExtension::ExecuteCreateImpulseResponse(const FToolMenuContext& MenuContext)
+{
+	UContentBrowserAssetContextMenuContext* Context = MenuContext.FindContext<UContentBrowserAssetContextMenuContext>();
+	if (!Context || Context->SelectedObjects.Num() == 0)
+	{
+		return;
+	}
+
+	const FString DefaultSuffix = TEXT("_IR");
+	FAssetToolsModule& AssetToolsModule = FModuleManager::Get().LoadModuleChecked<FAssetToolsModule>("AssetTools");
+
+	// Create the factory used to generate the asset
+	UAudioImpulseResponseFactory* Factory = NewObject<UAudioImpulseResponseFactory>();
+	
+	// only converts 0th selected object for now (see return statement)
+	for (const TWeakObjectPtr<UObject>& Object : Context->SelectedObjects)
+	{
+		// stage the soundwave on the factory to be used during asset creation
+		USoundWave* Wave = Cast<USoundWave>(Object);
+		check(Wave);
+		Factory->StagedSoundWave = Wave; // WeakPtr gets reset by the Factory after it is consumed
+
+		// Determine an appropriate name
+		FString Name;
+		FString PackagePath;
+		AssetToolsModule.Get().CreateUniqueAssetName(Wave->GetOutermost()->GetName(), DefaultSuffix, PackagePath, Name);
+
+		// create new asset
+		AssetToolsModule.Get().CreateAsset(Name, FPackageName::GetLongPackagePath(PackagePath), UAudioImpulseResponse::StaticClass(), Factory);
+	}
+}
+
+UAudioImpulseResponseFactory::UAudioImpulseResponseFactory(const FObjectInitializer& ObjectInitializer)
+	: Super(ObjectInitializer)
+{
+	SupportedClass = UAudioImpulseResponse::StaticClass();
+
+	bCreateNew = true;
+	bEditorImport = false;
+	bEditAfterNew = true;
+}
+
+UObject* UAudioImpulseResponseFactory::FactoryCreateNew(UClass* Class, UObject* InParent, FName InName, EObjectFlags Flags, UObject* Context, FFeedbackContext* Warn)
+{
+	UAudioImpulseResponse* NewAsset = NewObject<UAudioImpulseResponse>(InParent, InName, Flags);
+
+	if (StagedSoundWave.IsValid())
+	{
+		USoundWave* Wave = StagedSoundWave.Get();
+
+		constexpr bool bForceSynchronousLoad = true;
+
+		Loader.LoadSoundWave(Wave, [&](const USoundWave* SoundWave, const Audio::FSampleBuffer& LoadedSampleBuffer)
+		{
+			NewAsset->NumChannels = LoadedSampleBuffer.GetNumChannels();
+			NewAsset->SampleRate = LoadedSampleBuffer.GetSampleRate();
+			const int32 NumSamples = LoadedSampleBuffer.GetNumSamples();
+
+			NewAsset->ImpulseResponse.Reset();
+
+			if (NumSamples > 0)
+			{
+				NewAsset->ImpulseResponse.AddUninitialized(NumSamples);
+
+				// Convert to float.
+				const int16* InputBuffer = LoadedSampleBuffer.GetData();
+				float* OutputBuffer = NewAsset->ImpulseResponse.GetData();
+
+				for (int32 i = 0; i < NumSamples; ++i)
+				{
+					OutputBuffer[i] = static_cast<float>(InputBuffer[i]) / 32768.0f;
+				}
+			}
+		}, bForceSynchronousLoad);
+
+		Loader.Update();
+
+		StagedSoundWave.Reset();
+	}
+
+	return NewAsset;
+}
+
+#undef LOCTEXT_NAMESPACE
diff --git a/Engine/Plugins/Runtime/Synthesis/Source/SynthesisEditor/Private/ImpulseResponse.cpp b/Engine/Plugins/Runtime/Synthesis/Source/SynthesisEditor/Private/ImpulseResponse.cpp
deleted file mode 100644
diff --git a/Engine/Plugins/Runtime/Synthesis/Source/SynthesisEditor/Private/SoundWaveAssetActionExtenderConvolutionReverb.cpp b/Engine/Plugins/Runtime/Synthesis/Source/SynthesisEditor/Private/SoundWaveAssetActionExtenderConvolutionReverb.cpp
deleted file mode 100644
diff --git a/Engine/Plugins/Runtime/Synthesis/Source/SynthesisEditor/Private/SynthesisEditorModule.cpp b/Engine/Plugins/Runtime/Synthesis/Source/SynthesisEditor/Private/SynthesisEditorModule.cpp
--- a/Engine/Plugins/Runtime/Synthesis/Source/SynthesisEditor/Private/SynthesisEditorModule.cpp
+++ b/Engine/Plugins/Runtime/Synthesis/Source/SynthesisEditor/Private/SynthesisEditorModule.cpp
@@ -12,8 +12,7 @@
 #include "AudioEditorModule.h"
 #include "EpicSynth1PresetBank.h"
 #include "MonoWaveTablePresetBank.h"
-#include "ImpulseResponse.h"
-#include "SoundWaveAssetActionExtenderConvolutionReverb.h"
+#include "AudioImpulseResponseAsset.h"
 #include "ToolMenus.h"
 
 IMPLEMENT_MODULE(FSynthesisEditorModule, SynthesisEditor)
@@ -24,7 +23,7 @@ void FSynthesisEditorModule::StartupModule()
 
 	AssetTools.RegisterAssetTypeActions(MakeShareable(new FAssetTypeActions_ModularSynthPresetBank));
 	AssetTools.RegisterAssetTypeActions(MakeShareable(new FAssetTypeActions_MonoWaveTableSynthPreset));
-	AssetTools.RegisterAssetTypeActions(MakeShareable(new FAssetTypeActions_ImpulseResponse));
+	AssetTools.RegisterAssetTypeActions(MakeShareable(new FAssetTypeActions_AudioImpulseResponse));
 
 	// Now that we've loaded this module, we need to register our effect preset actions
 	IAudioEditorModule* AudioEditorModule = &FModuleManager::LoadModuleChecked<IAudioEditorModule>("AudioEditor");
@@ -39,5 +38,5 @@ void FSynthesisEditorModule::ShutdownModule()
 
 void FSynthesisEditorModule::RegisterMenus()
 {
-	FSoundWaveAssetActionExtenderConvolutionReverb::RegisterMenus();
-}
\ No newline at end of file
+	FAudioImpulseResponseExtension::RegisterMenus();
+}
diff --git a/Engine/Plugins/Runtime/Synthesis/Source/SynthesisEditor/Public/SoundWaveAssetActionExtenderConvolutionReverb.h b/Engine/Plugins/Runtime/Synthesis/Source/SynthesisEditor/Public/SoundWaveAssetActionExtenderConvolutionReverb.h
deleted file mode 100644
diff --git a/Engine/Source/Runtime/Engine/Classes/Sound/SoundEffectBase.h b/Engine/Source/Runtime/Engine/Classes/Sound/SoundEffectBase.h
--- a/Engine/Source/Runtime/Engine/Classes/Sound/SoundEffectBase.h
+++ b/Engine/Source/Runtime/Engine/Classes/Sound/SoundEffectBase.h
@@ -85,7 +85,7 @@ public:
 	bool IsPreset(USoundEffectPreset* InPreset) const;
 
 	/** Enqueues a lambda command on a thread safe queue which is pumped from the audio render thread. */
-	void EffectCommand(TFunction<void()> Command);
+	void EffectCommand(TUniqueFunction<void()> Command);
 
 protected:
 	FSoundEffectBase();
@@ -106,7 +106,7 @@ protected:
 	FThreadSafeBool bIsActive;
 
 	// Effect command queue
-	TQueue<TFunction<void()>> CommandQueue;
+	TQueue<TUniqueFunction<void()>> CommandQueue;
 
 	// Allow preset to re-register when editor update is requested
 	// and create effects using the templated Create call
diff --git a/Engine/Source/Runtime/Engine/Private/SoundEffectBase.cpp b/Engine/Source/Runtime/Engine/Private/SoundEffectBase.cpp
--- a/Engine/Source/Runtime/Engine/Private/SoundEffectBase.cpp
+++ b/Engine/Source/Runtime/Engine/Private/SoundEffectBase.cpp
@@ -51,7 +51,7 @@ bool FSoundEffectBase::IsPreset(USoundEffectPreset* InPreset) const
 	return Preset == InPreset;
 }
 
-void FSoundEffectBase::EffectCommand(TFunction<void()> Command)
+void FSoundEffectBase::EffectCommand(TUniqueFunction<void()> Command)
 {
 	CommandQueue.Enqueue(MoveTemp(Command));
 }
@@ -59,7 +59,7 @@ void FSoundEffectBase::EffectCommand(TFunction<void()> Command)
 void FSoundEffectBase::PumpPendingMessages()
 {
 	// Pumps the command queue
-	TFunction<void()> Command;
+	TUniqueFunction<void()> Command;
 	while (CommandQueue.Dequeue(Command))
 	{
 		Command();
diff --git a/Engine/Source/Runtime/SignalProcessing/Public/DSP/Dsp.h b/Engine/Source/Runtime/SignalProcessing/Public/DSP/Dsp.h
--- a/Engine/Source/Runtime/SignalProcessing/Public/DSP/Dsp.h
+++ b/Engine/Source/Runtime/SignalProcessing/Public/DSP/Dsp.h
@@ -503,9 +503,15 @@ namespace Audio
 			return false;
 		}
 
+		void CopyParams(T& OutParamsCopy) const
+		{
+			FScopeLock Lock(&CritSect);
+			OutParamsCopy = CurrentParams;
+		}
+
 		bool bChanged;
 		T CurrentParams;
-		FCriticalSection CritSect;
+		mutable FCriticalSection CritSect;
 	};
 
 	/**
