Adding convolution core redirects for new convolution reverb

diff --git a/Engine/Plugins/Runtime/Synthesis/Config/BaseSynthesis.ini b/Engine/Plugins/Runtime/Synthesis/Config/BaseSynthesis.ini
new file mode 100644
--- /dev/null
+++ b/Engine/Plugins/Runtime/Synthesis/Config/BaseSynthesis.ini
@@ -0,0 +1,15 @@
+[CoreRedirects]
+
++PropertyRedirects=(OldName="IRData", NewName="ImpulseResponse")
++PropertyRedirects=(OldName="AssetSampleRate", NewName="SampleRate")
++ClassRedirects=(OldName="ImpulseResponse", NewName="AudioImpulseResponse")
+
++PropertyRedirects=(OldName="ConvolutionReverbSettings.ImpulseResponse", NewName="SubmixEffectConvolutionReverbSettings.ImpulseResponse_DEPRECATED")
++PropertyRedirects=(OldName="ConvolutionReverbSettings.SurroundRearChannelBleedAmount", NewName="SubmixEffectConvolutionReverbSettings.SurroundRearChannelBleedAmount_DEPRECATED")
++PropertyRedirects=(OldName="ConvolutionReverbSettings.AllowHArdwareAcceleration", NewName="SubmixEffectConvolutionReverbSettings.AllowHardwareAcceleration_DEPRECATED")
+
++ClassRedirects=(OldName="ConvolutionReverbPreset", NewName="SubmixEffectConvolutionReverbPreset")
++StructRedirects=(OldName="ConvolutionReverbSettings", NewName="SubmixEffectConvolutionReverbSettings")
++PropertyRedirects=(OldName="ConvolutionReverbPreset.Settings", NewName="SubmixEffectConvolutionReverbPreset.Settings")
+
+
diff --git a/Engine/Plugins/Runtime/Synthesis/Source/Synthesis/Classes/SubmixEffects/SubmixEffectConvolutionReverb.h b/Engine/Plugins/Runtime/Synthesis/Source/Synthesis/Classes/SubmixEffects/SubmixEffectConvolutionReverb.h
--- a/Engine/Plugins/Runtime/Synthesis/Source/Synthesis/Classes/SubmixEffects/SubmixEffectConvolutionReverb.h
+++ b/Engine/Plugins/Runtime/Synthesis/Source/Synthesis/Classes/SubmixEffects/SubmixEffectConvolutionReverb.h
@@ -38,7 +38,6 @@ class SYNTHESIS_API UAudioImpulseResponse : public UObject
 
 public:
 	UAudioImpulseResponse();
-
 	/** The interleaved audio samples used in convolution. */
 	UPROPERTY()
 	TArray<float> ImpulseResponse;
@@ -88,6 +87,16 @@ struct SYNTHESIS_API FSubmixEffectConvolutionReverbSettings
 	/* If true, send Surround Rear Channel Bleed Amount sends front left to back right and vice versa */
 	UPROPERTY(EditAnywhere, BlueprintReadWrite, Category = SubmixEffectPreset)
 	bool bSurroundRearChannelFlip;
+
+	UPROPERTY(meta = ( DeprecatedProperty ) )
+	float SurroundRearChannelBleedAmount_DEPRECATED;
+
+	UPROPERTY(meta = ( DeprecatedProperty ) )
+	UAudioImpulseResponse* ImpulseResponse_DEPRECATED;
+
+	UPROPERTY(meta = ( DeprecatedProperty ) )
+	bool AllowHardwareAcceleration_DEPRECATED;
+
 };
 
 
diff --git a/Engine/Plugins/Runtime/Synthesis/Source/Synthesis/Private/SubmixEffectConvolutionReverb.cpp b/Engine/Plugins/Runtime/Synthesis/Source/Synthesis/Private/SubmixEffectConvolutionReverb.cpp
--- a/Engine/Plugins/Runtime/Synthesis/Source/Synthesis/Private/SubmixEffectConvolutionReverb.cpp
+++ b/Engine/Plugins/Runtime/Synthesis/Source/Synthesis/Private/SubmixEffectConvolutionReverb.cpp
@@ -215,6 +215,9 @@ FSubmixEffectConvolutionReverbSettings::FSubmixEffectConvolutionReverbSettings()
 ,	SurroundRearChannelBleedDb(-60.f)
 ,	bInvertRearChannelBleedPhase(false)
 ,	bSurroundRearChannelFlip(false)
+,	SurroundRearChannelBleedAmount_DEPRECATED(NAN)
+,	ImpulseResponse_DEPRECATED(nullptr)
+,	AllowHardwareAcceleration_DEPRECATED(true)
 {
 }
 
@@ -662,6 +665,24 @@ void USubmixEffectConvolutionReverbPreset::PostLoad()
 {
 	Super::PostLoad();
 
+	// This handles previous version 
+	if (FMath::IsFinite(Settings.SurroundRearChannelBleedAmount_DEPRECATED))
+	{
+		Settings.SurroundRearChannelBleedDb = Audio::ConvertToDecibels(FMath::Abs(Settings.SurroundRearChannelBleedAmount_DEPRECATED));
+		Settings.bInvertRearChannelBleedPhase = Settings.SurroundRearChannelBleedAmount_DEPRECATED < 0.f;
+	}
+
+	if (nullptr != Settings.ImpulseResponse_DEPRECATED)
+	{
+		ImpulseResponse = Settings.ImpulseResponse_DEPRECATED;
+	}
+
+	if (!Settings.AllowHardwareAcceleration_DEPRECATED)
+	{
+		bEnableHardwareAcceleration = false;
+	}
+
+
 	BindToImpulseResponseObjectChange();
 }
 
