Fixing uninitialized variable warning

diff --git a/Engine/Plugins/Runtime/Synthesis/Config/BaseSynthesis.ini b/Engine/Plugins/Runtime/Synthesis/Config/BaseSynthesis.ini
--- a/Engine/Plugins/Runtime/Synthesis/Config/BaseSynthesis.ini
+++ b/Engine/Plugins/Runtime/Synthesis/Config/BaseSynthesis.ini
@@ -1,15 +1,15 @@
 [CoreRedirects]
 
-+PropertyRedirects=(OldName="IRData", NewName="ImpulseResponse")
-+PropertyRedirects=(OldName="AssetSampleRate", NewName="SampleRate")
-+ClassRedirects=(OldName="ImpulseResponse", NewName="AudioImpulseResponse")
+;[4.25]
++StructRedirects=(OldName="ConvolutionReverbSettings", NewName="SubmixEffectConvolutionReverbSettings")
++ClassRedirects=(OldName="/Script/Synthesis.ConvolutionReverbPreset", NewName="/Script/Synthesis.SubmixEffectConvolutionReverbPreset")
++ClassRedirects=(OldName="/Script/Synthesis.ImpulseResponse", NewName="/Script/Synthesis.AudioImpulseResponse")
 
-+PropertyRedirects=(OldName="ConvolutionReverbSettings.ImpulseResponse", NewName="SubmixEffectConvolutionReverbSettings.ImpulseResponse_DEPRECATED")
-+PropertyRedirects=(OldName="ConvolutionReverbSettings.SurroundRearChannelBleedAmount", NewName="SubmixEffectConvolutionReverbSettings.SurroundRearChannelBleedAmount_DEPRECATED")
-+PropertyRedirects=(OldName="ConvolutionReverbSettings.AllowHArdwareAcceleration", NewName="SubmixEffectConvolutionReverbSettings.AllowHardwareAcceleration_DEPRECATED")
++PropertyRedirects=(OldName="/Script/Synthesis.ConvolutionReverbSettings.ImpulseResponse", NewName="/Script/Synthesis.SubmixEffectConvolutionReverbSettings.ImpulseResponse_DEPRECATED")
++PropertyRedirects=(OldName="/Script/Synthesis.IRData", NewName="/Script/Synthesis.ImpulseResponse")
++PropertyRedirects=(OldName="/Script/Synthesis.AssetSampleRate", NewName="/Script/Synthesis.SampleRate")
++PropertyRedirects=(OldName="/Script/Synthesis.ConvolutionReverbSettings.SurroundRearChannelBleedAmount", NewName="/Script/Synthesis.SubmixEffectConvolutionReverbSettings.SurroundRearChannelBleedAmount_DEPRECATED")
++PropertyRedirects=(OldName="/Script/Synthesis.ConvolutionReverbSettings.AllowHArdwareAcceleration", NewName="/Script/Synthesis.SubmixEffectConvolutionReverbSettings.AllowHardwareAcceleration_DEPRECATED")
 
-+ClassRedirects=(OldName="ConvolutionReverbPreset", NewName="SubmixEffectConvolutionReverbPreset")
-+StructRedirects=(OldName="ConvolutionReverbSettings", NewName="SubmixEffectConvolutionReverbSettings")
-+PropertyRedirects=(OldName="ConvolutionReverbPreset.Settings", NewName="SubmixEffectConvolutionReverbPreset.Settings")
 
 
diff --git a/Engine/Plugins/Runtime/Synthesis/Source/Synthesis/Private/SubmixEffectConvolutionReverb.cpp b/Engine/Plugins/Runtime/Synthesis/Source/Synthesis/Private/SubmixEffectConvolutionReverb.cpp
--- a/Engine/Plugins/Runtime/Synthesis/Source/Synthesis/Private/SubmixEffectConvolutionReverb.cpp
+++ b/Engine/Plugins/Runtime/Synthesis/Source/Synthesis/Private/SubmixEffectConvolutionReverb.cpp
@@ -215,7 +215,7 @@ FSubmixEffectConvolutionReverbSettings::FSubmixEffectConvolutionReverbSettings()
 ,	SurroundRearChannelBleedDb(-60.f)
 ,	bInvertRearChannelBleedPhase(false)
 ,	bSurroundRearChannelFlip(false)
-,	SurroundRearChannelBleedAmount_DEPRECATED(NAN)
+,	SurroundRearChannelBleedAmount_DEPRECATED(0.0f)
 ,	ImpulseResponse_DEPRECATED(nullptr)
 ,	AllowHardwareAcceleration_DEPRECATED(true)
 {
@@ -666,22 +666,30 @@ void USubmixEffectConvolutionReverbPreset::PostLoad()
 	Super::PostLoad();
 
 	// This handles previous version 
-	if (FMath::IsFinite(Settings.SurroundRearChannelBleedAmount_DEPRECATED))
+	if (Settings.SurroundRearChannelBleedAmount_DEPRECATED != 0.f)
 	{
 		Settings.SurroundRearChannelBleedDb = Audio::ConvertToDecibels(FMath::Abs(Settings.SurroundRearChannelBleedAmount_DEPRECATED));
 		Settings.bInvertRearChannelBleedPhase = Settings.SurroundRearChannelBleedAmount_DEPRECATED < 0.f;
+
+		Settings.SurroundRearChannelBleedAmount_DEPRECATED = 0.f;
+		Modify();
 	}
 
 	if (nullptr != Settings.ImpulseResponse_DEPRECATED)
 	{
 		ImpulseResponse = Settings.ImpulseResponse_DEPRECATED;
+		
+		Settings.ImpulseResponse_DEPRECATED = nullptr;
+		Modify();
 	}
 
 	if (!Settings.AllowHardwareAcceleration_DEPRECATED)
 	{
 		bEnableHardwareAcceleration = false;
-	}
 
+		Settings.AllowHardwareAcceleration_DEPRECATED = true;
+		Modify();
+	}
 
 	BindToImpulseResponseObjectChange();
 }
