Fixing Convolution Reverb crash on sample rate conversion of impulse response

diff --git a/Engine/Plugins/Runtime/Synthesis/Source/Synthesis/Private/SubmixEffectConvolutionReverb.cpp b/Engine/Plugins/Runtime/Synthesis/Source/Synthesis/Private/SubmixEffectConvolutionReverb.cpp
--- a/Engine/Plugins/Runtime/Synthesis/Source/Synthesis/Private/SubmixEffectConvolutionReverb.cpp
+++ b/Engine/Plugins/Runtime/Synthesis/Source/Synthesis/Private/SubmixEffectConvolutionReverb.cpp
@@ -54,30 +54,15 @@ namespace AudioConvReverbIntrinsics
 			return ConvolutionAlgorithm;
 		}
 
-		int32 NumFrames = InInitData.Samples.Num() / AlgoSettings.NumImpulseResponses;
 
-		// Prepare deinterleave pointers
-		TArray<Audio::AlignedFloatBuffer> DeinterleaveSamples;
-		while (DeinterleaveSamples.Num() < AlgoSettings.NumImpulseResponses)
-		{
-			Audio::AlignedFloatBuffer& Buffer = DeinterleaveSamples.Emplace_GetRef();
-			if (NumFrames > 0)
-			{
-				Buffer.AddUninitialized(NumFrames);
-			}
-		}
+		const TArray<float>* TargetImpulseSamples = &InInitData.Samples;
 
+		TArray<float> ResampledImpulseSamples;
 		// Prepare impulse samples by converting samplerate and deinterleaving.
-		if (InInitData.ImpulseSampleRate == InInitData.TargetSampleRate)
-		{
-			// Deinterleave impulse samples
-			Audio::FConvolutionReverb::DeinterleaveBuffer(DeinterleaveSamples, InInitData.Samples, AlgoSettings.NumImpulseResponses);
-		}
-		else	
+		if (InInitData.ImpulseSampleRate != InInitData.TargetSampleRate)
 		{
 			// convert sample rate of impulse 
-			TArray<float> TargetImpulseSamples;
-			float SampleRateRatio = InInitData.TargetSampleRate / InInitData.ImpulseSampleRate;
+			float SampleRateRatio = InInitData.ImpulseSampleRate / InInitData.TargetSampleRate;
 
 			TUniquePtr<Audio::ISampleRateConverter> Converter(Audio::ISampleRateConverter::CreateSampleRateConverter());
 
@@ -88,12 +73,27 @@ namespace AudioConvReverbIntrinsics
 			}
 
 			Converter->Init(SampleRateRatio, AlgoSettings.NumImpulseResponses);
-			Converter->ProcessFullbuffer(InInitData.Samples.GetData(), InInitData.Samples.Num(), TargetImpulseSamples);
+			Converter->ProcessFullbuffer(InInitData.Samples.GetData(), InInitData.Samples.Num(), ResampledImpulseSamples);
 
-			// Deinterleave impulse samples
-			Audio::FConvolutionReverb::DeinterleaveBuffer(DeinterleaveSamples, TargetImpulseSamples, AlgoSettings.NumImpulseResponses);
+			TargetImpulseSamples = &ResampledImpulseSamples;
 		}
 
+		const int32 NumFrames = TargetImpulseSamples->Num() / AlgoSettings.NumImpulseResponses;
+
+		// Prepare deinterleave pointers
+		TArray<Audio::AlignedFloatBuffer> DeinterleaveSamples;
+		while (DeinterleaveSamples.Num() < AlgoSettings.NumImpulseResponses)
+		{
+			Audio::AlignedFloatBuffer& Buffer = DeinterleaveSamples.Emplace_GetRef();
+			if (NumFrames > 0)
+			{
+				Buffer.AddUninitialized(NumFrames);
+			}
+		}
+
+		// Deinterleave impulse samples
+		Audio::FConvolutionReverb::DeinterleaveBuffer(DeinterleaveSamples, *TargetImpulseSamples, AlgoSettings.NumImpulseResponses);
+
 		// Set impulse responses in algorithm
 		for (int32 i = 0; i < DeinterleaveSamples.Num(); i++)
 		{
@@ -315,7 +315,12 @@ FSubmixEffectConvolutionReverb::FConvolutionAlgorithmInitData FSubmixEffectConvo
 
 	if (IRAssetData.NumChannels > 0)
 	{
-		CreateConvolutionSettings.AlgorithmSettings.MaxNumImpulseResponseSamples = (IRAssetData.Samples.Num() / IRAssetData.NumChannels) + 1;
+		float SampleRateRatio = 1.f;
+		if (IRAssetData.SampleRate > 0.f)
+		{
+			SampleRateRatio = SampleRate / IRAssetData.SampleRate;
+		}
+		CreateConvolutionSettings.AlgorithmSettings.MaxNumImpulseResponseSamples = FMath::CeilToInt(SampleRateRatio * IRAssetData.Samples.Num() / IRAssetData.NumChannels) + 256;
 	}
 
 	// Sample rate conversion settings
