Reinterleaving chunked audio from deprecated data.
Adding seperate array for deprecated format of IR samples
Routes block size setting to appropriate underlying algorithm implementation

diff --git a/Engine/Plugins/Runtime/Synthesis/Config/BaseSynthesis.ini b/Engine/Plugins/Runtime/Synthesis/Config/BaseSynthesis.ini
--- a/Engine/Plugins/Runtime/Synthesis/Config/BaseSynthesis.ini
+++ b/Engine/Plugins/Runtime/Synthesis/Config/BaseSynthesis.ini
@@ -6,7 +6,7 @@
 +ClassRedirects=(OldName="/Script/Synthesis.ImpulseResponse", NewName="/Script/Synthesis.AudioImpulseResponse")
 
 +PropertyRedirects=(OldName="/Script/Synthesis.ConvolutionReverbSettings.ImpulseResponse", NewName="/Script/Synthesis.SubmixEffectConvolutionReverbSettings.ImpulseResponse_DEPRECATED")
-+PropertyRedirects=(OldName="/Script/Synthesis.IRData", NewName="/Script/Synthesis.ImpulseResponse")
++PropertyRedirects=(OldName="/Script/Synthesis.ImpulseResponse.IRData", NewName="/Script/Synthesis.AudioImpulseResponse.IRData_DEPRECATED")
 +PropertyRedirects=(OldName="/Script/Synthesis.AssetSampleRate", NewName="/Script/Synthesis.SampleRate")
 +PropertyRedirects=(OldName="/Script/Synthesis.ConvolutionReverbSettings.SurroundRearChannelBleedAmount", NewName="/Script/Synthesis.SubmixEffectConvolutionReverbSettings.SurroundRearChannelBleedAmount_DEPRECATED")
 +PropertyRedirects=(OldName="/Script/Synthesis.ConvolutionReverbSettings.AllowHArdwareAcceleration", NewName="/Script/Synthesis.SubmixEffectConvolutionReverbSettings.AllowHardwareAcceleration_DEPRECATED")
diff --git a/Engine/Plugins/Runtime/Synthesis/Source/Synthesis/Classes/SubmixEffects/SubmixEffectConvolutionReverb.h b/Engine/Plugins/Runtime/Synthesis/Source/Synthesis/Classes/SubmixEffects/SubmixEffectConvolutionReverb.h
--- a/Engine/Plugins/Runtime/Synthesis/Source/Synthesis/Classes/SubmixEffects/SubmixEffectConvolutionReverb.h
+++ b/Engine/Plugins/Runtime/Synthesis/Source/Synthesis/Classes/SubmixEffects/SubmixEffectConvolutionReverb.h
@@ -54,6 +54,9 @@ public:
 	UPROPERTY(EditAnywhere, BlueprintReadWrite, Category = SubmixEffectPreset, meta = (ClamMin = "-60.0", ClampMax = "15.0"))
 	float NormalizationVolumeDb;
 
+	UPROPERTY(meta = ( DeprecatedProperty ) )
+	TArray<float> IRData_DEPRECATED;
+
 #if WITH_EDITORONLY_DATA
 	virtual void PostEditChangeProperty(FPropertyChangedEvent& PropertyChangedEvent) override;
 
@@ -280,16 +283,19 @@ public:
 
 	virtual void PreEditChange(FProperty* PropertyAboutToChange) override;
 	virtual void PostEditChangeProperty(FPropertyChangedEvent& PropertyChangedEvent) override;
-	virtual void PostLoad() override;
 
 	// Called when a property changes on teh ImpulseResponse object
 	void PostEditChangeImpulseProperty(FPropertyChangedEvent& PropertyChangedEvent);
 #endif
 
+	virtual void PostLoad() override;
+
 private:
 
 	void UpdateSettings();
 
+	void UpdateDeprecatedProperties();
+
 	// This method requires that the submix effect is registered with a preset.  If this 
 	// submix effect is not registered with a preset, then this will not update the convolution
 	// algorithm.
@@ -299,6 +305,7 @@ private:
 	FSubmixEffectConvolutionReverbSettings SettingsCopy; 
 
 #if WITH_EDITORONLY_DATA
+
 	TMap<UObject*, FDelegateHandle> DelegateHandles;
 #endif
 };
diff --git a/Engine/Plugins/Runtime/Synthesis/Source/Synthesis/Private/SubmixEffectConvolutionReverb.cpp b/Engine/Plugins/Runtime/Synthesis/Source/Synthesis/Private/SubmixEffectConvolutionReverb.cpp
--- a/Engine/Plugins/Runtime/Synthesis/Source/Synthesis/Private/SubmixEffectConvolutionReverb.cpp
+++ b/Engine/Plugins/Runtime/Synthesis/Source/Synthesis/Private/SubmixEffectConvolutionReverb.cpp
@@ -188,25 +188,6 @@ void UAudioImpulseResponse::PostEditChangeProperty(FPropertyChangedEvent& Proper
 	Super::PostEditChangeProperty(PropertyChangedEvent);
 
 	OnObjectPropertyChanged.Broadcast(PropertyChangedEvent);
-	/*
-	UObject* Outer = GetOuter();
-	if (nullptr == Outer)
-	{
-		return;
-	}
-
-	if (Outer->GetClass()->IsChildOf(USubmixEffectConvolutionReverbPreset::StaticClass()))
-	{
-		USubmixEffectConvolutionReverbPreset* Preset = CastChecked<USubmixEffectConvolutionReverbPreset>(Outer);
-
-		if (nullptr == Preset)
-		{
-			return;
-		}
-
-		Preset->PostEditChangeImpulseProperty(PropertyChangedEvent);
-	}
-	*/
 }
 #endif
 
@@ -307,7 +288,7 @@ FSubmixEffectConvolutionReverb::FConvolutionAlgorithmInitData FSubmixEffectConvo
 
 	// IConvolutionAlgorithm settings
 	CreateConvolutionSettings.AlgorithmSettings.bEnableHardwareAcceleration = IRAssetData.bEnableHardwareAcceleration;
-	CreateConvolutionSettings.AlgorithmSettings.BlockNumSamples = 256;
+	CreateConvolutionSettings.AlgorithmSettings.BlockNumSamples = IRAssetData.BlockSize;
 	CreateConvolutionSettings.AlgorithmSettings.NumInputChannels = ConvNumInputChannels;
 	CreateConvolutionSettings.AlgorithmSettings.NumOutputChannels = ConvNumOutputChannels;
 	CreateConvolutionSettings.AlgorithmSettings.NumImpulseResponses = IRAssetData.NumChannels;
@@ -457,7 +438,7 @@ FSubmixEffectConvolutionReverb::FVersionData FSubmixEffectConvolutionReverb::Upd
 {
 	// Copy data from preset into internal IRAssetData.
 	
-	check(IsInGameThread());
+	check(IsInGameThread() || IsInAudioThread());
 
 	FScopeLock IRAssetLock(&IRAssetCriticalSection);
 
@@ -666,38 +647,7 @@ void USubmixEffectConvolutionReverbPreset::PreEditChange(FProperty* PropertyAbou
 	}
 }
 
-void USubmixEffectConvolutionReverbPreset::PostLoad()
-{
-	Super::PostLoad();
-
-	// This handles previous version 
-	if (Settings.SurroundRearChannelBleedAmount_DEPRECATED != 0.f)
-	{
-		Settings.SurroundRearChannelBleedDb = Audio::ConvertToDecibels(FMath::Abs(Settings.SurroundRearChannelBleedAmount_DEPRECATED));
-		Settings.bInvertRearChannelBleedPhase = Settings.SurroundRearChannelBleedAmount_DEPRECATED < 0.f;
 
-		Settings.SurroundRearChannelBleedAmount_DEPRECATED = 0.f;
-		Modify();
-	}
-
-	if (nullptr != Settings.ImpulseResponse_DEPRECATED)
-	{
-		ImpulseResponse = Settings.ImpulseResponse_DEPRECATED;
-		
-		Settings.ImpulseResponse_DEPRECATED = nullptr;
-		Modify();
-	}
-
-	if (!Settings.AllowHardwareAcceleration_DEPRECATED)
-	{
-		bEnableHardwareAcceleration = false;
-
-		Settings.AllowHardwareAcceleration_DEPRECATED = true;
-		Modify();
-	}
-
-	BindToImpulseResponseObjectChange();
-}
 
 void USubmixEffectConvolutionReverbPreset::PostEditChangeProperty(FPropertyChangedEvent& PropertyChangedEvent)
 {
@@ -760,3 +710,86 @@ void USubmixEffectConvolutionReverbPreset::PostEditChangeImpulseProperty(FProper
 }
 
 #endif
+
+void USubmixEffectConvolutionReverbPreset::UpdateDeprecatedProperties()
+{
+	if (Settings.SurroundRearChannelBleedAmount_DEPRECATED != 0.f)
+	{
+		Settings.SurroundRearChannelBleedDb = Audio::ConvertToDecibels(FMath::Abs(Settings.SurroundRearChannelBleedAmount_DEPRECATED));
+		Settings.bInvertRearChannelBleedPhase = Settings.SurroundRearChannelBleedAmount_DEPRECATED < 0.f;
+
+		Settings.SurroundRearChannelBleedAmount_DEPRECATED = 0.f;
+		Modify();
+	}
+
+	if (nullptr != Settings.ImpulseResponse_DEPRECATED)
+	{
+		ImpulseResponse = Settings.ImpulseResponse_DEPRECATED;
+
+		// Need to make a copy of existing samples and reinterleave 
+		// The previous version had these sampels chunked by channel
+		// like [[all channel 0 samples][all channel 1 samples][...][allchannel N samples]]
+		//
+		// They need to be in interleave format to work in this class.
+
+		int32 NumChannels = Settings.ImpulseResponse_DEPRECATED->NumChannels;
+
+		if (NumChannels > 1)
+		{
+			const TArray<float>& IRData = Settings.ImpulseResponse_DEPRECATED->IRData_DEPRECATED;
+			int32 NumSamples = IRData.Num();
+
+			if (NumSamples > 0)
+			{
+				ImpulseResponse->ImpulseResponse.Reset();
+				ImpulseResponse->ImpulseResponse.AddUninitialized(NumSamples);
+
+				int32 NumFrames = NumSamples / NumChannels;
+
+				const float* InputSamples = IRData.GetData();
+				float* OutputSamples = ImpulseResponse->ImpulseResponse.GetData();
+
+				for (int32 ChannelIndex = 0; ChannelIndex < NumChannels; ChannelIndex++)
+				{
+					for (int32 FrameIndex = 0; FrameIndex < NumFrames; FrameIndex++)
+					{
+						OutputSamples[FrameIndex * NumChannels + ChannelIndex] = InputSamples[ChannelIndex * NumFrames + FrameIndex];
+					}
+				}
+			}
+		}
+		else
+		{
+			// Do not need to do anything for zero or one channels.
+			ImpulseResponse->ImpulseResponse = ImpulseResponse->IRData_DEPRECATED;
+		}
+
+		// Discard samples after they have been copied.
+		Settings.ImpulseResponse_DEPRECATED->IRData_DEPRECATED = TArray<float>();
+
+		Settings.ImpulseResponse_DEPRECATED = nullptr;
+
+		Modify();
+	}
+
+	if (!Settings.AllowHardwareAcceleration_DEPRECATED)
+	{
+		bEnableHardwareAcceleration = false;
+
+		Settings.AllowHardwareAcceleration_DEPRECATED = true;
+		Modify();
+	}
+}
+
+void USubmixEffectConvolutionReverbPreset::PostLoad()
+{
+	Super::PostLoad();
+
+	// This handles previous version 
+	UpdateDeprecatedProperties();
+
+#if WITH_EDITORONLY_DATA
+	// Bind to trigger new convolution algorithms when UAudioImpulseResponse object changes.
+	BindToImpulseResponseObjectChange();
+#endif
+}
