Fix for edge-case crash w/ synth sample wave table player.

Make the decode/pcm-loader work sync to avoid race-condition with other sounds using the same asset.

diff --git a/Engine/Plugins/Runtime/Synthesis/Source/Synthesis/Private/SynthComponentWaveTable.cpp b/Engine/Plugins/Runtime/Synthesis/Source/Synthesis/Private/SynthComponentWaveTable.cpp
--- a/Engine/Plugins/Runtime/Synthesis/Source/Synthesis/Private/SynthComponentWaveTable.cpp
+++ b/Engine/Plugins/Runtime/Synthesis/Source/Synthesis/Private/SynthComponentWaveTable.cpp
@@ -130,7 +130,7 @@ void USynthSamplePlayer::LoadSoundWaveInternal()
 			}
 		};
 
-		SoundWaveLoader.LoadSoundWave(SoundWave, MoveTemp(OnLoaded));
+		SoundWaveLoader.LoadSoundWave(SoundWave, MoveTemp(OnLoaded), true);
 	}
 }
 
