Remove Optimization.Never
#rnx

diff --git a/Engine/Plugins/Runtime/Synthesis/Source/Synthesis/Synthesis.Build.cs b/Engine/Plugins/Runtime/Synthesis/Source/Synthesis/Synthesis.Build.cs
--- a/Engine/Plugins/Runtime/Synthesis/Source/Synthesis/Synthesis.Build.cs
+++ b/Engine/Plugins/Runtime/Synthesis/Source/Synthesis/Synthesis.Build.cs
@@ -6,8 +6,6 @@ namespace UnrealBuildTool.Rules
 	{
 		public Synthesis(ReadOnlyTargetRules Target) : base(Target)
 		{
-			OptimizeCode = CodeOptimization.Never;
-
 			PublicDependencyModuleNames.AddRange(
 				new string[] {
 					"Core",
diff --git a/Engine/Source/Runtime/AudioExtensions/AudioExtensions.Build.cs b/Engine/Source/Runtime/AudioExtensions/AudioExtensions.Build.cs
--- a/Engine/Source/Runtime/AudioExtensions/AudioExtensions.Build.cs
+++ b/Engine/Source/Runtime/AudioExtensions/AudioExtensions.Build.cs
@@ -6,8 +6,6 @@ namespace UnrealBuildTool.Rules
 	{
 		public AudioExtensions(ReadOnlyTargetRules Target) : base(Target)
 		{
-			OptimizeCode = CodeOptimization.Never;
-
 			PublicDependencyModuleNames.AddRange(
 				new string[]
 				{
