Fixing type in convolution reverb. PR  #6823
#rnx

diff --git a/Engine/Plugins/Runtime/Synthesis/Source/SynthesisEditor/Classes/AudioImpulseResponseAsset.h b/Engine/Plugins/Runtime/Synthesis/Source/SynthesisEditor/Classes/AudioImpulseResponseAsset.h
--- a/Engine/Plugins/Runtime/Synthesis/Source/SynthesisEditor/Classes/AudioImpulseResponseAsset.h
+++ b/Engine/Plugins/Runtime/Synthesis/Source/SynthesisEditor/Classes/AudioImpulseResponseAsset.h
@@ -14,7 +14,7 @@
 class FAssetTypeActions_AudioImpulseResponse : public FAssetTypeActions_Base
 {
 public:
-	virtual FText GetName() const override { return NSLOCTEXT("AssetTypeActions", "AssetTypeActions_AudioImpulseResponse", "Audio Impulse Respons"); }
+	virtual FText GetName() const override { return NSLOCTEXT("AssetTypeActions", "AssetTypeActions_AudioImpulseResponse", "Audio Impulse Response"); }
 	virtual FColor GetTypeColor() const override { return FColor(0, 255, 0); }
 	virtual UClass* GetSupportedClass() const override;
 	virtual const TArray<FText>& GetSubMenus() const override;
