Added convenience individual carrier/base value setters to upgraded effects (BitCrush & Chorus) as well as base settings setter

diff --git a/Engine/Plugins/Runtime/Synthesis/Source/Synthesis/Classes/SourceEffects/SourceEffectBitCrusher.h b/Engine/Plugins/Runtime/Synthesis/Source/Synthesis/Classes/SourceEffects/SourceEffectBitCrusher.h
--- a/Engine/Plugins/Runtime/Synthesis/Source/Synthesis/Classes/SourceEffects/SourceEffectBitCrusher.h
+++ b/Engine/Plugins/Runtime/Synthesis/Source/Synthesis/Classes/SourceEffects/SourceEffectBitCrusher.h
@@ -8,39 +8,46 @@
 #include "SourceEffectBitCrusher.generated.h"
 
 
+// Forward Declarations
+class USoundModulatorBase;
+
+USTRUCT(BlueprintType)
+struct SYNTHESIS_API FSourceEffectBitCrusherBaseSettings
+{
+	GENERATED_USTRUCT_BODY()
+
+	// The reduced frequency to use for the audio stream. 
+	UPROPERTY(EditAnywhere, BlueprintReadWrite, Category = "SourceEffect|Preset", meta = (ClampMin = "0.1", ClampMax = "96000.0", UIMin = "500.0", UIMax = "16000.0"))
+	float SampleRate = 8000.0f;
+
+	// The reduced bit depth to use for the audio stream
+	UPROPERTY(EditAnywhere, BlueprintReadWrite, Category = "SourceEffect|Preset", meta = (ClampMin = "1.0", ClampMax = "24.0", UIMin = "1.0", UIMax = "16.0"))
+	float BitDepth = 8.0f;
+};
+
 USTRUCT(BlueprintType)
 struct SYNTHESIS_API FSourceEffectBitCrusherSettings
 {
 	GENERATED_USTRUCT_BODY()
 
-#if WITH_EDITORONLY_DATA
 	UPROPERTY(meta = (PropertyDeprecated))
-	float CrushedSampleRate;
-#endif // if WITH_EDITORONLY_DATA
+	float CrushedSampleRate = 8000.0f;
 
 	// The reduced frequency to use for the audio stream. 
-	UPROPERTY(EditAnywhere, BlueprintReadWrite, Category = "SourceEffect|Preset", meta = (DisplayName = "Sample Rate", AudioParam = "SampleRate", UIMin = "500.0", UIMax = "16000.0"))
+	UPROPERTY(EditAnywhere, BlueprintReadWrite, Category = "SourceEffect|Preset", meta = (DisplayName = "Sample Rate", AudioParam = "SampleRate", ClampMin = "0.1", ClampMax = "96000.0", UIMin = "500.0", UIMax = "16000.0"))
 	FSoundModulationDestinationSettings SampleRateModulation;
 
-#if WITH_EDITORONLY_DATA
 	UPROPERTY(meta = (PropertyDeprecated))
-	float CrushedBits;
-#endif // if WITH_EDITORONLY_DATA
+	float CrushedBits = 8.0f;
 
 	// The reduced bit depth to use for the audio stream
 	UPROPERTY(EditAnywhere, BlueprintReadWrite, Category = "SourceEffect|Preset", meta = (DisplayName = "Bit Depth", AudioParam = "BitDepth", ClampMin = "1.0", ClampMax = "24.0", UIMin = "1.0", UIMax = "16.0"))
 	FSoundModulationDestinationSettings BitModulation;
 
 	FSourceEffectBitCrusherSettings()
-#if WITH_EDITORONLY_DATA
-		: CrushedSampleRate(8000.0f)
-		, CrushedBits(8.0f)
-#endif // if WITH_EDITORONLY_DATA
 	{
-#if WITH_EDITORONLY_DATA
 		SampleRateModulation.Value = CrushedSampleRate;
 		BitModulation.Value = CrushedBits;
-#endif // WITH_EDITORONLY_DATA
 	}
 };
 
@@ -56,9 +63,9 @@ public:
 	// Process the input block of audio. Called on audio thread.
 	virtual void ProcessAudio(const FSoundEffectSourceInputData& InData, float* OutAudioBufferData) override;
 
-	void SetSampleRateModulator(const FSoundModulationDestinationSettings& InModulatorSettings);
+	void SetSampleRateModulator(const USoundModulatorBase* Modulator);
 
-	void SetBitModulator(const FSoundModulationDestinationSettings& InModulatorSettings);
+	void SetBitModulator(const USoundModulatorBase* Modulator);
 
 protected:
 	Audio::FBitCrusher BitCrusher;
@@ -83,15 +90,25 @@ public:
 
 	virtual FColor GetPresetColor() const override { return FColor(196.0f, 185.0f, 121.0f); }
 
-	UFUNCTION(BlueprintCallable, Category = "Audio|Effects")
-	void SetBitModulator(const FSoundModulationDestinationSettings& InModulatorSettings);
+	UFUNCTION(BlueprintCallable, Category = "Audio|Effects|BitCrusher")
+	void SetBits(float Bits);
+
+	UFUNCTION(BlueprintCallable, Category = "Audio|Effects|BitCrusher")
+	void SetBitModulator(const USoundModulatorBase* Modulator);
+
+	UFUNCTION(BlueprintCallable, Category = "Audio|Effects|BitCrusher")
+	void SetSampleRate(float SampleRate);
+
+	UFUNCTION(BlueprintCallable, Category = "Audio|Effects|BitCrusher")
+	void SetSampleRateModulator(const USoundModulatorBase* Modulator);
 
-	UFUNCTION(BlueprintCallable, Category = "Audio|Effects")
-	void SetSampleRateModulator(const FSoundModulationDestinationSettings& InModulatorSettings);
+	// Sets just base (i.e. carrier) setting values without modifying modulation source references
+	UFUNCTION(BlueprintCallable, Category = "Audio|Effects|BitCrusher")
+	void SetSettings(const FSourceEffectBitCrusherBaseSettings& Settings);
 
-	UFUNCTION(BlueprintCallable, Category = "Audio|Effects")
-	void SetSettings(const FSourceEffectBitCrusherSettings& InSettings);
+	UFUNCTION(BlueprintCallable, Category = "Audio|Effects|BitCrusher")
+	void SetModulationSettings(const FSourceEffectBitCrusherSettings& ModulationSettings);
 
-	UPROPERTY(EditAnywhere, BlueprintReadOnly, Category = "Audio|Effects")
+	UPROPERTY(EditAnywhere, BlueprintReadOnly, Category = "SourceEffect|Preset")
 	FSourceEffectBitCrusherSettings Settings;
 };
diff --git a/Engine/Plugins/Runtime/Synthesis/Source/Synthesis/Classes/SourceEffects/SourceEffectChorus.h b/Engine/Plugins/Runtime/Synthesis/Source/Synthesis/Classes/SourceEffects/SourceEffectChorus.h
--- a/Engine/Plugins/Runtime/Synthesis/Source/Synthesis/Classes/SourceEffects/SourceEffectChorus.h
+++ b/Engine/Plugins/Runtime/Synthesis/Source/Synthesis/Classes/SourceEffects/SourceEffectChorus.h
@@ -9,59 +9,83 @@
 #include "SourceEffectChorus.generated.h"
 
 USTRUCT(BlueprintType)
-struct SYNTHESIS_API FSourceEffectChorusSettings
+struct SYNTHESIS_API FSourceEffectChorusBaseSettings
 {
 	GENERATED_USTRUCT_BODY()
 
 	// The depth of the chorus effect
-	UPROPERTY(meta = (DeprecatedProperty))
-	float Depth;
-
 	UPROPERTY(EditAnywhere, BlueprintReadWrite, Category = "SourceEffect|Preset", meta = (ClampMin = "0.0", ClampMax = "1.0", UIMin = "0.0", UIMax = "1.0"))
-	FSoundModulationDestinationSettings DepthModulation;
+	float Depth = 0.2f;
 
 	// The frequency of the chorus effect
+	UPROPERTY(EditAnywhere, BlueprintReadWrite, Category = "SourceEffect|Preset", meta = (UIMin = "0.0", UIMax = "5.0"))
+	float Frequency = 2.0f;
+
+	// The feedback of the chorus effect
+	UPROPERTY(EditAnywhere, BlueprintReadWrite, Category = "SourceEffect|Preset", meta = (ClampMin = "0.0", ClampMax = "1.0", UIMin = "0.0", UIMax = "1.0"))
+	float Feedback = 0.3f;
+
+	// The wet level of the chorus effect
+	UPROPERTY(EditAnywhere, BlueprintReadWrite, Category = "SourceEffect|Preset", meta = (ClampMin = "0.0", ClampMax = "1.0", UIMin = "0.0", UIMax = "1.0"))
+	float WetLevel = 0.5f;
+
+	// The dry level of the chorus effect
+	UPROPERTY(EditAnywhere, BlueprintReadWrite, Category = "SourceEffect|Preset", meta = (ClampMin = "0.0", ClampMax = "1.0", UIMin = "0.0", UIMax = "1.0"))
+	float DryLevel = 0.5f;
+
+	// The spread of the effect (larger means greater difference between left and right delay lines)
+	UPROPERTY(EditAnywhere, BlueprintReadWrite, Category = "SourceEffect|Preset", meta = (ClampMin = "0.0", ClampMax = "1.0", UIMin = "0.0", UIMax = "1.0"))
+	float Spread = 0.0f;
+};
+
+USTRUCT(BlueprintType)
+struct SYNTHESIS_API FSourceEffectChorusSettings
+{
+	GENERATED_USTRUCT_BODY()
+
 	UPROPERTY(meta = (DeprecatedProperty))
-	float Frequency;
+	float Depth = 0.2f;
 
-	UPROPERTY(EditAnywhere, BlueprintReadWrite, Category = "SourceEffect|Preset", meta = (DisplayName = "Frequency", AudioParam = "LowRateFrequency", UIMin = "0.0", UIMax = "5.0"))
-	FSoundModulationDestinationSettings FrequencyModulation;
+	UPROPERTY(meta = (DeprecatedProperty))
+	float Frequency = 2.0f;
 
 	UPROPERTY(meta = (DeprecatedProperty))
-	float Feedback;
+	float Feedback = 0.3f;
+
+	UPROPERTY(meta = (DeprecatedProperty))
+	float WetLevel = 0.5f;
+
+	UPROPERTY(meta = (DeprecatedProperty))
+	float DryLevel = 0.5f;
+
+	UPROPERTY(meta = (DeprecatedProperty))
+	float Spread = 0.0f;
+
+	// The depth of the chorus effect
+	UPROPERTY(EditAnywhere, BlueprintReadWrite, Category = "SourceEffect|Preset", meta = (DisplayName = "Depth", ClampMin = "0.0", ClampMax = "1.0", UIMin = "0.0", UIMax = "1.0"))
+	FSoundModulationDestinationSettings DepthModulation;
+
+	// The frequency of the chorus effect
+	UPROPERTY(EditAnywhere, BlueprintReadWrite, Category = "SourceEffect|Preset", meta = (DisplayName = "Frequency", AudioParam = "LowRateFrequency", UIMin = "0.0", UIMax = "5.0"))
+	FSoundModulationDestinationSettings FrequencyModulation;
 
 	// The feedback of the chorus effect
 	UPROPERTY(EditAnywhere, BlueprintReadWrite, Category = "SourceEffect|Preset", meta = (DisplayName = "Feedback", ClampMin = "0.0", ClampMax = "1.0", UIMin = "0.0", UIMax = "1.0"))
 	FSoundModulationDestinationSettings FeedbackModulation;
 
-	UPROPERTY(meta = (DeprecatedProperty))
-	float WetLevel;
-
 	// The wet level of the chorus effect
 	UPROPERTY(EditAnywhere, BlueprintReadWrite, Category = "SourceEffect|Preset", meta = (DisplayName = "Wet Level", ClampMin = "0.0", ClampMax = "1.0", UIMin = "0.0", UIMax = "1.0"))
 	FSoundModulationDestinationSettings WetModulation;
 
-	UPROPERTY(meta = (DeprecatedProperty))
-	float DryLevel;
-
 	// The dry level of the chorus effect
 	UPROPERTY(EditAnywhere, BlueprintReadWrite, Category = "SourceEffect|Preset", meta = (DisplayName = "Dry Level", ClampMin = "0.0", ClampMax = "1.0", UIMin = "0.0", UIMax = "1.0"))
 	FSoundModulationDestinationSettings DryModulation;
 
-	UPROPERTY(meta = (DeprecatedProperty))
-	float Spread;
-
 	// The spread of the effect (larger means greater difference between left and right delay lines)
 	UPROPERTY(EditAnywhere, BlueprintReadWrite, Category = "SourceEffect|Preset", meta = (DisplayName = "Spread", ClampMin = "0.0", ClampMax = "1.0", UIMin = "0.0", UIMax = "1.0"))
 	FSoundModulationDestinationSettings SpreadModulation;
 
 	FSourceEffectChorusSettings()
-		: Depth(0.2f)
-		, Frequency(2.0f)
-		, Feedback(0.3f)
-		, WetLevel(0.5f)
-		, DryLevel(0.5f)
-		, Spread(0.0f)
 	{
 		DepthModulation.Value = 0.2f;
 		FrequencyModulation.Value = 2.0f;
@@ -84,12 +108,12 @@ public:
 	// Process the input block of audio. Called on audio thread.
 	virtual void ProcessAudio(const FSoundEffectSourceInputData& InData, float* OutAudioBufferData) override;
 
-	void SetDepthModulator(const FSoundModulationDestinationSettings& InModulatorSettings);
-	void SetFeedbackModulator(const FSoundModulationDestinationSettings& InModulatorSettings);
-	void SetFrequencyModulator(const FSoundModulationDestinationSettings& InModulatorSettings);
-	void SetWetModulator(const FSoundModulationDestinationSettings& InModulatorSettings);
-	void SetDryModulator(const FSoundModulationDestinationSettings& InModulatorSettings);
-	void SetSpreadModulator(const FSoundModulationDestinationSettings& InModulatorSettings);
+	void SetDepthModulator(const USoundModulatorBase* InModulator);
+	void SetFeedbackModulator(const USoundModulatorBase* InModulator);
+	void SetFrequencyModulator(const USoundModulatorBase* InModulator);
+	void SetWetModulator(const USoundModulatorBase* InModulator);
+	void SetDryModulator(const USoundModulatorBase* InModulator);
+	void SetSpreadModulator(const USoundModulatorBase* InModulator);
 
 protected:
 	Audio::FChorus Chorus;
@@ -116,28 +140,49 @@ public:
 
 	virtual void OnInit() override;
 
-	UFUNCTION(BlueprintCallable, Category = "Audio|Effects")
-	void SetDepthModulator(const FSoundModulationDestinationSettings& InModulatorSettings);
+	UFUNCTION(BlueprintCallable, Category = "Audio|Effects|Chorus")
+	void SetDepth(float Depth);
 
-	UFUNCTION(BlueprintCallable, Category = "Audio|Effects")
-	void SetFeedbackModulator(const FSoundModulationDestinationSettings& InModulatorSettings);
+	UFUNCTION(BlueprintCallable, Category = "Audio|Effects|Chorus")
+	void SetDepthModulator(const USoundModulatorBase* Modulator);
 
-	UFUNCTION(BlueprintCallable, Category = "Audio|Effects")
-	void SetFrequencyModulator(const FSoundModulationDestinationSettings& InModulatorSettings);
+	UFUNCTION(BlueprintCallable, Category = "Audio|Effects|Chorus")
+	void SetFeedback(float Feedback);
 
-	UFUNCTION(BlueprintCallable, Category = "Audio|Effects")
-	void SetWetModulator(const FSoundModulationDestinationSettings& InModulatorSettings);
+	UFUNCTION(BlueprintCallable, Category = "Audio|Effects|Chorus")
+	void SetFeedbackModulator(const USoundModulatorBase* Modulator);
 
-	UFUNCTION(BlueprintCallable, Category = "Audio|Effects")
-	void SetDryModulator(const FSoundModulationDestinationSettings& InModulatorSettings);
+	UFUNCTION(BlueprintCallable, Category = "Audio|Effects|Chorus")
+	void SetFrequency(float Frequency);
 
-	UFUNCTION(BlueprintCallable, Category = "Audio|Effects")
-	void SetSpreadModulator(const FSoundModulationDestinationSettings& InModulatorSettings);
+	UFUNCTION(BlueprintCallable, Category = "Audio|Effects|Chorus")
+	void SetFrequencyModulator(const USoundModulatorBase* Modulator);
 
-	UFUNCTION(BlueprintCallable, Category = "Audio|Effects")
-	void SetSettings(const FSourceEffectChorusSettings& InSettings);
+	UFUNCTION(BlueprintCallable, Category = "Audio|Effects|Chorus")
+	void SetWet(float WetAmount);
 
-	// The depth of the chorus effect
-	UPROPERTY(EditAnywhere, BlueprintReadWrite, Category = "SourceEffect|Preset")
+	UFUNCTION(BlueprintCallable, Category = "Audio|Effects|Chorus")
+	void SetWetModulator(const USoundModulatorBase* Modulator);
+
+	UFUNCTION(BlueprintCallable, Category = "Audio|Effects|Chorus")
+	void SetDry(float DryAmount);
+
+	UFUNCTION(BlueprintCallable, Category = "Audio|Effects|Chorus")
+	void SetDryModulator(const USoundModulatorBase* Modulator);
+
+	UFUNCTION(BlueprintCallable, Category = "Audio|Effects|Chorus")
+	void SetSpread(float Spread);
+
+	UFUNCTION(BlueprintCallable, Category = "Audio|Effects|Chorus")
+	void SetSpreadModulator(const USoundModulatorBase* Modulator);
+
+	// Sets just base (i.e. carrier) setting values without modifying modulation source references
+	UFUNCTION(BlueprintCallable, Category = "Audio|Effects|Chorus")
+	void SetSettings(const FSourceEffectChorusBaseSettings& Settings);
+
+	UFUNCTION(BlueprintCallable, Category = "Audio|Effects|Chorus")
+	void SetModulationSettings(const FSourceEffectChorusSettings& ModulationSettings);
+
+	UPROPERTY(EditAnywhere, BlueprintReadOnly, Category = "SourceEffect|Preset")
 	FSourceEffectChorusSettings Settings;
 };
\ No newline at end of file
diff --git a/Engine/Plugins/Runtime/Synthesis/Source/Synthesis/Private/SourceEffectBitCrusher.cpp b/Engine/Plugins/Runtime/Synthesis/Source/Synthesis/Private/SourceEffectBitCrusher.cpp
--- a/Engine/Plugins/Runtime/Synthesis/Source/Synthesis/Private/SourceEffectBitCrusher.cpp
+++ b/Engine/Plugins/Runtime/Synthesis/Source/Synthesis/Private/SourceEffectBitCrusher.cpp
@@ -2,6 +2,9 @@
 
 #include "SourceEffects/SourceEffectBitCrusher.h"
 
+#include "IAudioModulation.h"
+
+
 void FSourceEffectBitCrusher::Init(const FSoundEffectSourceInitData& InitData)
 {
 	bIsActive = true;
@@ -12,8 +15,8 @@ void FSourceEffectBitCrusher::Init(const FSoundEffectSourceInitData& InitData)
 		BitMod.Init(InitData.AudioDeviceId, FName("BitDepth"), false /* bInIsBuffered */);
 		SampleRateMod.Init(InitData.AudioDeviceId, FName("SampleRate"), false /* bInIsBuffered */);
 
-		SetBitModulator(ProcPreset->Settings.BitModulation);
-		SetSampleRateModulator(ProcPreset->Settings.SampleRateModulation);
+		SetBitModulator(ProcPreset->Settings.BitModulation.Modulator);
+		SetSampleRateModulator(ProcPreset->Settings.SampleRateModulation.Modulator);
 	}
 }
 
@@ -39,20 +42,20 @@ void FSourceEffectBitCrusher::ProcessAudio(const FSoundEffectSourceInputData& In
 	BitCrusher.ProcessAudio(InData.InputSourceEffectBufferPtr, InData.NumSamples, OutAudioBufferData);
 }
 
-void FSourceEffectBitCrusher::SetBitModulator(const FSoundModulationDestinationSettings& InModulatorSettings)
+void FSourceEffectBitCrusher::SetBitModulator(const USoundModulatorBase* InModulator)
 {
-	BitMod.UpdateSettings(InModulatorSettings);
+	BitMod.UpdateModulator(InModulator);
 }
 
-void FSourceEffectBitCrusher::SetSampleRateModulator(const FSoundModulationDestinationSettings& InModulatorSettings)
+void FSourceEffectBitCrusher::SetSampleRateModulator(const USoundModulatorBase* InModulator)
 {
-	SampleRateMod.UpdateSettings(InModulatorSettings);
+	SampleRateMod.UpdateModulator(InModulator);
 }
 
 void USourceEffectBitCrusherPreset::OnInit()
 {
-	SetBitModulator(Settings.BitModulation);
-	SetSampleRateModulator(Settings.SampleRateModulation);
+	SetBitModulator(Settings.BitModulation.Modulator);
+	SetSampleRateModulator(Settings.SampleRateModulation.Modulator);
 }
 
 void USourceEffectBitCrusherPreset::Serialize(FArchive& Ar)
@@ -77,24 +80,51 @@ void USourceEffectBitCrusherPreset::Serialize(FArchive& Ar)
 #endif // WITH_EDITORONLY_DATA
 }
 
-void USourceEffectBitCrusherPreset::SetBitModulator(const FSoundModulationDestinationSettings& InModulatorSettings)
+void USourceEffectBitCrusherPreset::SetBits(float InBits)
+{
+	UpdateSettings([NewBits = InBits](FSourceEffectBitCrusherSettings& OutSettings)
+	{
+		OutSettings.BitModulation.Value = NewBits;
+	});
+}
+
+void USourceEffectBitCrusherPreset::SetBitModulator(const USoundModulatorBase* InModulator)
+{
+	IterateEffects<FSourceEffectBitCrusher>([InModulator](FSourceEffectBitCrusher& InCrusher)
+	{
+		InCrusher.SetBitModulator(InModulator);
+	});
+}
+
+void USourceEffectBitCrusherPreset::SetSampleRate(float InSampleRate)
 {
-	IterateEffects<FSourceEffectBitCrusher>([InModulatorSettings](FSourceEffectBitCrusher& InCrusher)
+	UpdateSettings([NewSampleRate = InSampleRate](FSourceEffectBitCrusherSettings& OutSettings)
 	{
-		InCrusher.SetBitModulator(InModulatorSettings);
+		OutSettings.SampleRateModulation.Value = NewSampleRate;
 	});
 }
 
-void USourceEffectBitCrusherPreset::SetSampleRateModulator(const FSoundModulationDestinationSettings& InModulatorSettings)
+void USourceEffectBitCrusherPreset::SetSampleRateModulator(const USoundModulatorBase* InModulator)
 {
-	IterateEffects<FSourceEffectBitCrusher>([InModulatorSettings](FSourceEffectBitCrusher& InCrusher)
+	IterateEffects<FSourceEffectBitCrusher>([InModulator](FSourceEffectBitCrusher& InCrusher)
 	{
-		InCrusher.SetSampleRateModulator(InModulatorSettings);
+		InCrusher.SetSampleRateModulator(InModulator);
 	});
 }
 
-void USourceEffectBitCrusherPreset::SetSettings(const FSourceEffectBitCrusherSettings& InSettings)
+void USourceEffectBitCrusherPreset::SetSettings(const FSourceEffectBitCrusherBaseSettings& InSettings)
 {
-	UpdateSettings(InSettings);
+	UpdateSettings([NewBaseSettings = InSettings](FSourceEffectBitCrusherSettings& OutSettings)
+	{
+		OutSettings.BitModulation.Value = NewBaseSettings.BitDepth;
+		OutSettings.SampleRateModulation.Value = NewBaseSettings.SampleRate;
+	});
+}
+
+void USourceEffectBitCrusherPreset::SetModulationSettings(const FSourceEffectBitCrusherSettings& InModulationSettings)
+{
+	UpdateSettings(InModulationSettings);
+
+	// Must be called to update modulators
 	OnInit();
 }
\ No newline at end of file
diff --git a/Engine/Plugins/Runtime/Synthesis/Source/Synthesis/Private/SourceEffectChorus.cpp b/Engine/Plugins/Runtime/Synthesis/Source/Synthesis/Private/SourceEffectChorus.cpp
--- a/Engine/Plugins/Runtime/Synthesis/Source/Synthesis/Private/SourceEffectChorus.cpp
+++ b/Engine/Plugins/Runtime/Synthesis/Source/Synthesis/Private/SourceEffectChorus.cpp
@@ -72,96 +72,159 @@ void FSourceEffectChorus::ProcessAudio(const FSoundEffectSourceInputData& InData
 	}
 }
 
-void FSourceEffectChorus::SetDepthModulator(const FSoundModulationDestinationSettings& InModulatorSettings)
+void FSourceEffectChorus::SetDepthModulator(const USoundModulatorBase* InModulator)
 {
-	DepthMod.UpdateSettings(InModulatorSettings);
+	DepthMod.UpdateModulator(InModulator);
 }
 
-void FSourceEffectChorus::SetFeedbackModulator(const FSoundModulationDestinationSettings& InModulatorSettings)
+void FSourceEffectChorus::SetFeedbackModulator(const USoundModulatorBase* InModulator)
 {
-	FeedbackMod.UpdateSettings(InModulatorSettings);
+	FeedbackMod.UpdateModulator(InModulator);
 }
 
-void FSourceEffectChorus::SetSpreadModulator(const FSoundModulationDestinationSettings& InModulatorSettings)
+void FSourceEffectChorus::SetSpreadModulator(const USoundModulatorBase* InModulator)
 {
-	SpreadMod.UpdateSettings(InModulatorSettings);
+	SpreadMod.UpdateModulator(InModulator);
 }
 
-void FSourceEffectChorus::SetDryModulator(const FSoundModulationDestinationSettings& InModulatorSettings)
+void FSourceEffectChorus::SetDryModulator(const USoundModulatorBase* InModulator)
 {
-	DryMod.UpdateSettings(InModulatorSettings);
+	DryMod.UpdateModulator(InModulator);
 }
 
-void FSourceEffectChorus::SetWetModulator(const FSoundModulationDestinationSettings& InModulatorSettings)
+void FSourceEffectChorus::SetWetModulator(const USoundModulatorBase* InModulator)
 {
-	WetMod.UpdateSettings(InModulatorSettings);
+	WetMod.UpdateModulator(InModulator);
 }
 
-void FSourceEffectChorus::SetFrequencyModulator(const FSoundModulationDestinationSettings& InModulatorSettings)
+void FSourceEffectChorus::SetFrequencyModulator(const USoundModulatorBase* InModulator)
 {
-	FrequencyMod.UpdateSettings(InModulatorSettings);
+	FrequencyMod.UpdateModulator(InModulator);
 }
 
 void USourceEffectChorusPreset::OnInit()
 {
-	SetDepthModulator(Settings.SpreadModulation);
-	SetDryModulator(Settings.DryModulation);
-	SetWetModulator(Settings.WetModulation);
-	SetFeedbackModulator(Settings.FeedbackModulation);
-	SetFrequencyModulator(Settings.FrequencyModulation);
-	SetSpreadModulator(Settings.SpreadModulation);
+	SetDepthModulator(Settings.SpreadModulation.Modulator);
+	SetDryModulator(Settings.DryModulation.Modulator);
+	SetWetModulator(Settings.WetModulation.Modulator);
+	SetFeedbackModulator(Settings.FeedbackModulation.Modulator);
+	SetFrequencyModulator(Settings.FrequencyModulation.Modulator);
+	SetSpreadModulator(Settings.SpreadModulation.Modulator);
 }
 
-void USourceEffectChorusPreset::SetDepthModulator(const FSoundModulationDestinationSettings& InModulatorSettings)
+void USourceEffectChorusPreset::SetDepth(float InDepth)
 {
-	IterateEffects<FSourceEffectChorus>([InModulatorSettings](FSourceEffectChorus& InDelay)
+	UpdateSettings([NewDepth = InDepth](FSourceEffectChorusSettings& OutSettings)
 	{
-		InDelay.SetDepthModulator(InModulatorSettings);
+		OutSettings.DepthModulation.Value = NewDepth;
 	});
 }
 
-void USourceEffectChorusPreset::SetFeedbackModulator(const FSoundModulationDestinationSettings& InModulatorSettings)
+void USourceEffectChorusPreset::SetDepthModulator(const USoundModulatorBase* InModulator)
 {
-	IterateEffects<FSourceEffectChorus>([InModulatorSettings](FSourceEffectChorus& InDelay)
+	IterateEffects<FSourceEffectChorus>([InModulator](FSourceEffectChorus& InDelay)
 	{
-		InDelay.SetFeedbackModulator(InModulatorSettings);
+		InDelay.SetDepthModulator(InModulator);
 	});
 }
 
-void USourceEffectChorusPreset::SetFrequencyModulator(const FSoundModulationDestinationSettings& InModulatorSettings)
+void USourceEffectChorusPreset::SetFeedback(float InFeedback)
 {
-	IterateEffects<FSourceEffectChorus>([InModulatorSettings](FSourceEffectChorus& InDelay)
+	UpdateSettings([NewFeedback = InFeedback](FSourceEffectChorusSettings& OutSettings)
 	{
-		InDelay.SetFrequencyModulator(InModulatorSettings);
+		OutSettings.FeedbackModulation.Value = NewFeedback;
 	});
 }
 
-void USourceEffectChorusPreset::SetWetModulator(const FSoundModulationDestinationSettings& InModulatorSettings)
+void USourceEffectChorusPreset::SetFeedbackModulator(const USoundModulatorBase* InModulator)
 {
-	IterateEffects<FSourceEffectChorus>([InModulatorSettings](FSourceEffectChorus& InDelay)
+	IterateEffects<FSourceEffectChorus>([InModulator](FSourceEffectChorus& InDelay)
 	{
-		InDelay.SetWetModulator(InModulatorSettings);
+		InDelay.SetFeedbackModulator(InModulator);
 	});
 }
 
-void USourceEffectChorusPreset::SetDryModulator(const FSoundModulationDestinationSettings& InModulatorSettings)
+void USourceEffectChorusPreset::SetFrequency(float InFrequency)
 {
-	IterateEffects<FSourceEffectChorus>([InModulatorSettings](FSourceEffectChorus& InDelay)
+	UpdateSettings([NewFrequency = InFrequency](FSourceEffectChorusSettings& OutSettings)
 	{
-		InDelay.SetDryModulator(InModulatorSettings);
+		OutSettings.FrequencyModulation.Value = NewFrequency;
 	});
 }
 
-void USourceEffectChorusPreset::SetSpreadModulator(const FSoundModulationDestinationSettings& InModulatorSettings)
+void USourceEffectChorusPreset::SetFrequencyModulator(const USoundModulatorBase* InModulator)
 {
-	IterateEffects<FSourceEffectChorus>([InModulatorSettings](FSourceEffectChorus& InDelay)
+	IterateEffects<FSourceEffectChorus>([InModulator](FSourceEffectChorus& InDelay)
 	{
-		InDelay.SetSpreadModulator(InModulatorSettings);
+		InDelay.SetFrequencyModulator(InModulator);
 	});
 }
 
-void USourceEffectChorusPreset::SetSettings(const FSourceEffectChorusSettings& InSettings)
+void USourceEffectChorusPreset::SetWet(float InWet)
 {
-	UpdateSettings(InSettings);
+	UpdateSettings([NewWet = InWet](FSourceEffectChorusSettings& OutSettings)
+	{
+		OutSettings.WetModulation.Value = NewWet;
+	});
+}
+
+void USourceEffectChorusPreset::SetWetModulator(const USoundModulatorBase* InModulator)
+{
+	IterateEffects<FSourceEffectChorus>([InModulator](FSourceEffectChorus& InDelay)
+	{
+		InDelay.SetWetModulator(InModulator);
+	});
+}
+
+void USourceEffectChorusPreset::SetDry(float InDry)
+{
+	UpdateSettings([NewDry = InDry](FSourceEffectChorusSettings& OutSettings)
+	{
+		OutSettings.DryModulation.Value = NewDry;
+	});
+}
+
+void USourceEffectChorusPreset::SetDryModulator(const USoundModulatorBase* InModulator)
+{
+	IterateEffects<FSourceEffectChorus>([InModulator](FSourceEffectChorus& InDelay)
+	{
+		InDelay.SetDryModulator(InModulator);
+	});
+}
+
+void USourceEffectChorusPreset::SetSpread(float InSpread)
+{
+	UpdateSettings([NewSpread = InSpread](FSourceEffectChorusSettings& OutSettings)
+	{
+		OutSettings.SpreadModulation.Value = NewSpread;
+	});
+}
+
+void USourceEffectChorusPreset::SetSpreadModulator(const USoundModulatorBase* InModulator)
+{
+	IterateEffects<FSourceEffectChorus>([InModulator](FSourceEffectChorus& InDelay)
+	{
+		InDelay.SetSpreadModulator(InModulator);
+	});
+}
+
+void USourceEffectChorusPreset::SetSettings(const FSourceEffectChorusBaseSettings& InSettings)
+{
+	UpdateSettings([NewBaseSettings = InSettings](FSourceEffectChorusSettings& OutSettings)
+	{
+		OutSettings.DepthModulation.Value = NewBaseSettings.Depth;
+		OutSettings.FrequencyModulation.Value = NewBaseSettings.Frequency;
+		OutSettings.FeedbackModulation.Value = NewBaseSettings.Feedback;
+		OutSettings.WetModulation.Value = NewBaseSettings.WetLevel;
+		OutSettings.DryModulation.Value = NewBaseSettings.DryLevel;
+		OutSettings.SpreadModulation.Value = NewBaseSettings.Spread;
+	});
+}
+
+void USourceEffectChorusPreset::SetModulationSettings(const FSourceEffectChorusSettings& InModulationSettings)
+{
+	UpdateSettings(InModulationSettings);
+
+	// Must be called to update modulators
 	OnInit();
 }
\ No newline at end of file
diff --git a/Engine/Source/Runtime/AudioMixer/Private/AudioMixerSourceManager.cpp b/Engine/Source/Runtime/AudioMixer/Private/AudioMixerSourceManager.cpp
--- a/Engine/Source/Runtime/AudioMixer/Private/AudioMixerSourceManager.cpp
+++ b/Engine/Source/Runtime/AudioMixer/Private/AudioMixerSourceManager.cpp
@@ -699,19 +699,19 @@ namespace Audio
 
 		FModulationDestination VolumeModulation;
 		VolumeModulation.Init(MixerDevice->DeviceID, FName("Volume"), false /* bInIsBuffered */, true /* bInValueLinear */);
-		VolumeModulation.UpdateSettings(InitParams.VolumeModulationSettings);
+		VolumeModulation.UpdateModulator(InitParams.VolumeModulationSettings.Modulator);
 
 		FModulationDestination PitchModulation;
 		PitchModulation.Init(MixerDevice->DeviceID, FName("Pitch"), false /* bInIsBuffered */);
-		PitchModulation.UpdateSettings(InitParams.PitchModulationSettings);
+		PitchModulation.UpdateModulator(InitParams.PitchModulationSettings.Modulator);
 
 		FModulationDestination HighpassModulation;
 		HighpassModulation.Init(MixerDevice->DeviceID, FName("HPFCutoffFrequency"), false /* bInIsBuffered */);
-		HighpassModulation.UpdateSettings(InitParams.HighpassModulationSettings);
+		HighpassModulation.UpdateModulator(InitParams.HighpassModulationSettings.Modulator);
 
 		FModulationDestination LowpassModulation;
 		LowpassModulation.Init(MixerDevice->DeviceID, FName("LPFCutoffFrequency"), false /* bInIsBuffered */);
-		LowpassModulation.UpdateSettings(InitParams.LowpassModulationSettings);
+		LowpassModulation.UpdateModulator(InitParams.LowpassModulationSettings.Modulator);
 
 		AudioMixerThreadCommand([this, SourceId, InitParams, VolumeModulation, HighpassModulation, LowpassModulation, PitchModulation]()
 		{
diff --git a/Engine/Source/Runtime/Engine/Classes/Sound/SoundEffectBase.h b/Engine/Source/Runtime/Engine/Classes/Sound/SoundEffectBase.h
--- a/Engine/Source/Runtime/Engine/Classes/Sound/SoundEffectBase.h
+++ b/Engine/Source/Runtime/Engine/Classes/Sound/SoundEffectBase.h
@@ -54,6 +54,12 @@ class USoundEffectPreset;
 			SettingsCopy = InSettings; \
 			Update(); \
 		} \
+		void UpdateSettings(TUniqueFunction<void(F##EFFECT_NAME##Settings&)> InCommand) \
+		{ \
+			FScopeLock ScopeLock(&SettingsCritSect); \
+			InCommand(SettingsCopy); \
+			Update(); \
+		} \
 		F##EFFECT_NAME##Settings GetSettings() \
 		{ \
 			FScopeLock ScopeLock(&SettingsCritSect); \
diff --git a/Engine/Source/Runtime/Engine/Classes/Sound/SoundModulationDestination.h b/Engine/Source/Runtime/Engine/Classes/Sound/SoundModulationDestination.h
--- a/Engine/Source/Runtime/Engine/Classes/Sound/SoundModulationDestination.h
+++ b/Engine/Source/Runtime/Engine/Classes/Sound/SoundModulationDestination.h
@@ -73,7 +73,7 @@ namespace Audio
 			 */
 			bool ProcessControl(float InValueUnitBase, int32 InNumSamples = 0);
 
-			void UpdateSettings(const FSoundModulationDestinationSettings& InSettings);
+			void UpdateModulator(const USoundModulatorBase* InModulator);
 
 		private:
 			FDeviceId DeviceId = INDEX_NONE;
diff --git a/Engine/Source/Runtime/Engine/Private/SoundModulationDestination.cpp b/Engine/Source/Runtime/Engine/Private/SoundModulationDestination.cpp
--- a/Engine/Source/Runtime/Engine/Private/SoundModulationDestination.cpp
+++ b/Engine/Source/Runtime/Engine/Private/SoundModulationDestination.cpp
@@ -216,9 +216,9 @@ namespace Audio
 		return !FMath::IsNearlyEqual(LastTarget, ValueTarget);
 	}
 
-	void FModulationDestination::UpdateSettings(const FSoundModulationDestinationSettings& InSettings)
+	void FModulationDestination::UpdateModulator(const USoundModulatorBase* InModulator)
 	{
-		const TWeakObjectPtr<const USoundModulatorBase> ModPtr(InSettings.Modulator);
+		const TWeakObjectPtr<const USoundModulatorBase> ModPtr(InModulator);
 		auto UpdateHandleLambda = [this, ModPtr]()
 		{
 			if (FAudioDevice* AudioDevice = FAudioDeviceManager::Get()->GetAudioDeviceRaw(DeviceId))
