Add submix listener side-chaining to MultibandCompressor submix effect

diff --git a/Engine/Plugins/Runtime/Synthesis/Source/Synthesis/Classes/SubmixEffects/SubmixEffectMultiBandCompressor.h b/Engine/Plugins/Runtime/Synthesis/Source/Synthesis/Classes/SubmixEffects/SubmixEffectMultiBandCompressor.h
--- a/Engine/Plugins/Runtime/Synthesis/Source/Synthesis/Classes/SubmixEffects/SubmixEffectMultiBandCompressor.h
+++ b/Engine/Plugins/Runtime/Synthesis/Source/Synthesis/Classes/SubmixEffects/SubmixEffectMultiBandCompressor.h
@@ -9,82 +9,110 @@
 #include "SubmixEffectMultiBandCompressor.generated.h"
 
 USTRUCT(BlueprintType)
-struct FDynamicsBandSettings
+struct SYNTHESIS_API FDynamicsBandSettings
 {
 	GENERATED_BODY()
 
 	// Frequency of the crossover between this band and the next. The last band will have this property ignored
-	UPROPERTY(EditAnywhere, BlueprintReadWrite, Category = "SourceEffect|Preset", meta = (ClampMin = "20.0", ClampMax = "20000.0", UIMin = "20.0", UIMax = "20000"))
+	UPROPERTY(EditAnywhere, BlueprintReadWrite, Category = Response, meta = (ClampMin = "20.0", ClampMax = "20000.0", UIMin = "20.0", UIMax = "20000"))
 	float CrossoverTopFrequency = 20000.f;
 
 	// The amount of time to ramp into any dynamics processing effect in milliseconds.
-	UPROPERTY(EditAnywhere, BlueprintReadWrite, Category = "SourceEffect|Preset", meta = (ClampMin = "1.0", ClampMax = "300.0", UIMin = "1.0", UIMax = "200.0"))
+	UPROPERTY(EditAnywhere, BlueprintReadWrite, Category = Response, meta = (ClampMin = "1.0", ClampMax = "300.0", UIMin = "1.0", UIMax = "200.0"))
 	float AttackTimeMsec = 10.f;
 
 	// The amount of time to release the dynamics processing effect in milliseconds
-	UPROPERTY(EditAnywhere, BlueprintReadWrite, Category = "SourceEffect|Preset", meta = (ClampMin = "20.0", ClampMax = "5000.0", UIMin = "20.0", UIMax = "5000.0"))
+	UPROPERTY(EditAnywhere, BlueprintReadWrite, Category = Response, meta = (ClampMin = "20.0", ClampMax = "5000.0", UIMin = "20.0", UIMax = "5000.0"))
 	float ReleaseTimeMsec = 100.f;
 
 	// The threshold at which to perform a dynamics processing operation
-	UPROPERTY(EditAnywhere, BlueprintReadWrite, Category = "SourceEffect|Preset", meta = (ClampMin = "-72.0", ClampMax = "0.0", UIMin = "-72.0", UIMax = "0.0"))
+	UPROPERTY(EditAnywhere, BlueprintReadWrite, Category = Dynamics, meta = (ClampMin = "-72.0", ClampMax = "0.0", UIMin = "-72.0", UIMax = "0.0"))
 	float ThresholdDb = -6.f;
 
 	// The dynamics processor ratio -- has different meaning depending on the processor type.
-	UPROPERTY(EditAnywhere, BlueprintReadWrite, Category = "SourceEffect|Preset", meta = (ClampMin = "1.0", ClampMax = "20.0", UIMin = "1.0", UIMax = "20.0"))
+	UPROPERTY(EditAnywhere, BlueprintReadWrite, Category = Dynamics, meta = (ClampMin = "1.0", ClampMax = "20.0", UIMin = "1.0", UIMax = "20.0"))
 	float Ratio = 1.5f;
 
 	// The knee bandwidth of the compressor to use in dB
-	UPROPERTY(EditAnywhere, BlueprintReadWrite, Category = "SourceEffect|Preset", meta = (ClampMin = "0.0", ClampMax = "20.0", UIMin = "0.0", UIMax = "20.0"))
+	UPROPERTY(EditAnywhere, BlueprintReadWrite, Category = Dynamics, meta = (ClampMin = "0.0", ClampMax = "20.0", UIMin = "0.0", UIMax = "20.0"))
 	float KneeBandwidthDb = 10.f;
 
 	// The input gain of the dynamics processor in dB
-	UPROPERTY(EditAnywhere, BlueprintReadWrite, Category = "SourceEffect|Preset", meta = (ClampMin = "-12.0", ClampMax = "20.0", UIMin = "-12.0", UIMax = "20.0"))
+	UPROPERTY(EditAnywhere, BlueprintReadWrite, Category = Gain, meta = (ClampMin = "-12.0", ClampMax = "20.0", UIMin = "-12.0", UIMax = "20.0"))
 	float InputGainDb = 0.f;
 
 	// The output gain of the dynamics processor in dB
-	UPROPERTY(EditAnywhere, BlueprintReadWrite, Category = "SourceEffect|Preset", meta = (ClampMin = "0.0", ClampMax = "20.0", UIMin = "0.0", UIMax = "20.0"))
+	UPROPERTY(EditAnywhere, BlueprintReadWrite, Category = Gain, meta = (ClampMin = "-12.0", ClampMax = "20.0", UIMin = "-12.0", UIMax = "20.0"))
 	float OutputGainDb = 0.f;
 };
 
 // A submix dynamics processor
 USTRUCT(BlueprintType)
-struct FSubmixEffectMultibandCompressorSettings
+struct SYNTHESIS_API FSubmixEffectMultibandCompressorSettings
 {
 	GENERATED_BODY()
 
 	// Controls how each band will react to audio above its threshold
-	UPROPERTY(EditAnywhere, BlueprintReadWrite, Category = "SourceEffect|Preset")
+	UPROPERTY(EditAnywhere, BlueprintReadWrite, Category = General, meta = (DisplayName = "Type"))
 	ESubmixEffectDynamicsProcessorType DynamicsProcessorType = ESubmixEffectDynamicsProcessorType::Compressor;
 
 	// Controls how quickly the bands will react to a signal above its threshold
-	UPROPERTY(EditAnywhere, BlueprintReadWrite, Category = "SourceEffect|Preset")
+	UPROPERTY(EditAnywhere, BlueprintReadWrite, Category = Dynamics)
 	ESubmixEffectDynamicsPeakMode PeakMode = ESubmixEffectDynamicsPeakMode::RootMeanSquared;
 
+	// Whether to compress all channels equally, and how to evaluate the overall level
+	UPROPERTY(EditAnywhere, BlueprintReadWrite, Category = Dynamics)
+	ESubmixEffectDynamicsChannelLinkMode LinkMode = ESubmixEffectDynamicsChannelLinkMode::Average;
+
 	// The amount of time to look ahead of the current audio. Allows for transients to be included in dynamics processing.
-	UPROPERTY(EditAnywhere, BlueprintReadWrite, Category = "SourceEffect|Preset",  meta = (ClampMin = "0.0", ClampMax = "50.0", UIMin = "0.0", UIMax = "50.0"))
+	UPROPERTY(EditAnywhere, BlueprintReadWrite, Category = Response, meta = (DisplayName = "Look Ahead (ms)", ClampMin = "0.0", ClampMax = "50.0", UIMin = "0.0", UIMax = "50.0"))
 	float LookAheadMsec = 3.f;
 
-	// Whether or not to average all channels of audio before inputing into the dynamics processor
-	UPROPERTY(EditAnywhere, BlueprintReadWrite, Category = "SourceEffect|Preset")
-	bool bLinkChannels = true;
-
 	// Toggles treating the attack and release envelopes as analog-style vs digital-style. Analog will respond a bit more naturally/slower.
-	UPROPERTY(EditAnywhere, BlueprintReadWrite, Category = "SourceEffect|Preset")
+	UPROPERTY(EditAnywhere, BlueprintReadWrite, Category = Response)
 	bool bAnalogMode = true;
 
 	// Turning off FourPole mode will use cheaper, shallower 2-pole crossovers
-	UPROPERTY(EditAnywhere, BlueprintReadWrite, Category = "SourceEffect|Preset")
+	UPROPERTY(EditAnywhere, BlueprintReadWrite, Category = Crossover)
 	bool bFourPole = true;
 
+	// Whether or not to bypass effect
+	UPROPERTY(EditAnywhere, BlueprintReadWrite, Category = General, meta = (DisplayName = "Bypass", DisplayAfter = "DynamicsProcessorType"))
+	bool bBypass;
+
+	UPROPERTY(EditAnywhere, BlueprintReadWrite, Category = Sidechain, meta = (EditCondition = "!bBypass"))
+	ESubmixEffectDynamicsKeySource KeySource = ESubmixEffectDynamicsKeySource::Default;
+
+	// If set, uses output of provided submix as modulator of input signal for dynamics processor (Uses input signal as default modulator)
+	UPROPERTY(EditAnywhere, BlueprintReadWrite, Category = Sidechain, meta = (EditCondition = "!bBypass && KeySource == ESubmixEffectDynamicsKeySource::AudioBus", EditConditionHides))
+	TObjectPtr<UAudioBus> ExternalAudioBus = nullptr;
+
+	// If set, uses output of provided submix as modulator of input signal for dynamics processor (Uses input signal as default modulator)
+	UPROPERTY(EditAnywhere, BlueprintReadWrite, Category = Sidechain, meta = (EditCondition = "!bBypass && KeySource == ESubmixEffectDynamicsKeySource::Submix", EditConditionHides))
+	TObjectPtr<USoundSubmix> ExternalSubmix = nullptr;
+
+	// Gain to apply to key signal if external input is supplied
+	UPROPERTY(EditAnywhere, BlueprintReadWrite, Category = Sidechain, meta = (DisplayName = "External Input Gain (dB)", EditCondition = "ExternalSubmix != nullptr || ExternalAudioBus != nullptr", UIMin = "-60.0", UIMax = "30.0"))
+	float KeyGainDb;
+
+	// Audition the key modulation signal, bypassing enveloping and processing the input signal.
+	UPROPERTY(EditAnywhere, BlueprintReadWrite, Category = Sidechain, meta = (DisplayName = "Key Audition", EditCondition = "!bBypass"))
+	bool bKeyAudition;
+
 	// Each band is a full dynamics processor, affecting at a unique frequency range
-	UPROPERTY(EditAnywhere, BlueprintReadWrite, Category = "SourceEffect|Preset")
+	UPROPERTY(EditAnywhere, BlueprintReadWrite, Category = Bands)
 	TArray<FDynamicsBandSettings> Bands;
 };
 
-class FSubmixEffectMultibandCompressor : public FSoundEffectSubmix
+class SYNTHESIS_API FSubmixEffectMultibandCompressor : public FSoundEffectSubmix
 {
 public:
-	FSubmixEffectMultibandCompressor() {};
+	FSubmixEffectMultibandCompressor();
+
+	virtual ~FSubmixEffectMultibandCompressor();
+
+	// Gets the effect's deviceId that owns it
+	Audio::FDeviceId GetDeviceId() const;
 
 	// Called on an audio effect at initialization on main thread before audio processing begins.
 	virtual void Init(const FSoundEffectSubmixInitData& InSampleRate) override;
@@ -99,14 +127,36 @@ public:
 	void Initialize(FSubmixEffectMultibandCompressorSettings& Settings);
 
 protected:
-	TArray<float> AudioInputFrame;
-	TArray<float> AudioOutputFrame;
+	// ISubmixBufferListener interface
+	void ResetKey();
+	void UpdateKeyFromSettings(const FSubmixEffectMultibandCompressorSettings& InSettings);
+	bool UpdateKeySourcePatch();
+
+	void OnDeviceCreated(Audio::FDeviceId InDeviceId);
+	void OnDeviceDestroyed(Audio::FDeviceId InDeviceId);
+
+	Audio::FMixerDevice* GetMixerDevice();
+
+	Audio::FAlignedFloatBuffer AudioExternal;
+
+private:
+	FKeySource KeySource;
+
+	Audio::FDeviceId DeviceId;
+
+	FDelegateHandle DeviceCreatedHandle;
+	FDelegateHandle DeviceDestroyedHandle;
+
 	TArray<Audio::FDynamicsProcessor> DynamicsProcessors;
+
+	Audio::FLinkwitzRileyBandSplitter KeyBandSplitter;
+	Audio::FMultibandBuffer KeyMultiBandBuffer;
+
 	Audio::FLinkwitzRileyBandSplitter BandSplitter;
 	Audio::FMultibandBuffer MultiBandBuffer;
 
 	// cached crossover + band info, so we can check if they need a re-build when editing
-	int32 PrevNumBands = 0; 
+	int32 PrevNumBands = 0;
 	TArray<float> PrevCrossovers;
 	bool bPrevFourPole = true;
 
@@ -115,20 +165,44 @@ protected:
 	float SampleRate = 48000.f;
 
 	bool bInitialized = false;
+	bool bBypass = false;
+
+	friend class USubmixEffectMultibandCompressorPreset;
 };
 
 UCLASS(ClassGroup = AudioSourceEffect, meta = (BlueprintSpawnableComponent))
-class USubmixEffectMultibandCompressorPreset : public USoundEffectSubmixPreset
+class SYNTHESIS_API USubmixEffectMultibandCompressorPreset : public USoundEffectSubmixPreset
 {
 	GENERATED_BODY()
 
 public:
-
 	EFFECT_PRESET_METHODS(SubmixEffectMultibandCompressor)
 
+	virtual void OnInit() override;
+
+#if WITH_EDITOR
+	virtual void PostEditChangeChainProperty(struct FPropertyChangedChainEvent& InChainEvent) override;
+#endif // WITH_EDITOR
+
+	UFUNCTION(BlueprintCallable, Category = "Audio|Effects")
+	void ResetKey();
+
+	// Sets the source key input as the provided AudioBus' output.  If no object is provided, key is set
+	// to effect's input.
+	UFUNCTION(BlueprintCallable, Category = "Audio|Effects")
+	void SetAudioBus(UAudioBus* AudioBus);
+
+	// Sets the source key input as the provided Submix's output.  If no object is provided, key is set
+	// to effect's input.
+	UFUNCTION(BlueprintCallable, Category = "Audio|Effects")
+	void SetExternalSubmix(USoundSubmix* Submix);
+
 	UFUNCTION(BlueprintCallable, Category = "Audio|Effects")
 	void SetSettings(const FSubmixEffectMultibandCompressorSettings& InSettings);
 
-	UPROPERTY(EditAnywhere, BlueprintReadWrite, Category = SubmixEffectPreset)
+	UPROPERTY(EditAnywhere, BlueprintReadWrite, Category = SubmixEffectPreset, meta = (ShowOnlyInnerProperties))
 	FSubmixEffectMultibandCompressorSettings Settings;
+
+private:
+	void SetKey(ESubmixEffectDynamicsKeySource InKeySource, UObject* InObject, int32 InNumChannels = 0);
 };
\ No newline at end of file
diff --git a/Engine/Plugins/Runtime/Synthesis/Source/Synthesis/Private/SubmixEffectMultiBandCompressor.cpp b/Engine/Plugins/Runtime/Synthesis/Source/Synthesis/Private/SubmixEffectMultiBandCompressor.cpp
--- a/Engine/Plugins/Runtime/Synthesis/Source/Synthesis/Private/SubmixEffectMultiBandCompressor.cpp
+++ b/Engine/Plugins/Runtime/Synthesis/Source/Synthesis/Private/SubmixEffectMultiBandCompressor.cpp
@@ -1,15 +1,85 @@
 // Copyright Epic Games, Inc. All Rights Reserved.
 
 #include "SubmixEffects/SubmixEffectMultiBandCompressor.h"
+
+#include "AudioDeviceManager.h"
+#include "AudioMixerDevice.h"
+#include "AudioMixerSubmix.h"
 #include "DSP/BufferVectorOperations.h"
 
+FSubmixEffectMultibandCompressor::FSubmixEffectMultibandCompressor()
+{
+	DeviceCreatedHandle = FAudioDeviceManagerDelegates::OnAudioDeviceCreated.AddRaw(this, &FSubmixEffectMultibandCompressor::OnDeviceCreated);
+	DeviceDestroyedHandle = FAudioDeviceManagerDelegates::OnAudioDeviceDestroyed.AddRaw(this, &FSubmixEffectMultibandCompressor::OnDeviceDestroyed);
+}
+
+FSubmixEffectMultibandCompressor::~FSubmixEffectMultibandCompressor()
+{
+	ResetKey();
+
+	FAudioDeviceManagerDelegates::OnAudioDeviceCreated.Remove(DeviceCreatedHandle);
+	FAudioDeviceManagerDelegates::OnAudioDeviceDestroyed.Remove(DeviceDestroyedHandle);
+}
+
+Audio::FDeviceId FSubmixEffectMultibandCompressor::GetDeviceId() const
+{
+	return DeviceId;
+}
+
 void FSubmixEffectMultibandCompressor::Init(const FSoundEffectSubmixInitData& InitData)
 {
 	NumChannels = 0;
 	SampleRate = InitData.SampleRate;
 
-	AudioOutputFrame.Reset();
 	MultiBandBuffer.Init(4, 2 * 1024);
+	KeyMultiBandBuffer.Init(4, 2 * 1024);
+
+	DeviceId = InitData.DeviceID;
+
+	if (USubmixEffectMultibandCompressorPreset* ProcPreset = Cast<USubmixEffectMultibandCompressorPreset>(Preset.Get()))
+	{
+		switch (ProcPreset->Settings.KeySource)
+		{
+			case ESubmixEffectDynamicsKeySource::AudioBus:
+			{
+				if (UAudioBus* AudioBus = ProcPreset->Settings.ExternalAudioBus)
+				{
+					KeySource.Update(ESubmixEffectDynamicsKeySource::AudioBus, AudioBus->GetUniqueID(), static_cast<int32>(AudioBus->AudioBusChannels) + 1);
+				}
+			}
+			break;
+
+			case ESubmixEffectDynamicsKeySource::Submix:
+			{
+				if (USoundSubmix* Submix = ProcPreset->Settings.ExternalSubmix)
+				{
+					KeySource.Update(ESubmixEffectDynamicsKeySource::Submix, Submix->GetUniqueID());
+				}
+			}
+			break;
+
+			default:
+			{
+				// KeySource is this effect's submix/input, so do nothing
+			}
+			break;
+		}
+	}
+}
+
+void FSubmixEffectMultibandCompressor::ResetKey()
+{
+	KeySource.Reset();
+}
+
+Audio::FMixerDevice* FSubmixEffectMultibandCompressor::GetMixerDevice()
+{
+	if (FAudioDeviceManager* DeviceManager = FAudioDeviceManager::Get())
+	{
+		return static_cast<Audio::FMixerDevice*>(DeviceManager->GetAudioDeviceRaw(DeviceId));
+	}
+
+	return nullptr;
 }
 
 void FSubmixEffectMultibandCompressor::OnPresetChanged()
@@ -21,10 +91,10 @@ void FSubmixEffectMultibandCompressor::OnPresetChanged()
 		return;
 	}
 
-	bool bNeedsReinit = (Settings.Bands.Num() != PrevNumBands 
-						 || Settings.Bands.Num() - 1 != PrevCrossovers.Num()
-						 || Settings.bFourPole != bPrevFourPole
-						 || bInitialized == false);
+	bool bNeedsReinit = (bInitialized == false
+		|| Settings.Bands.Num() != PrevNumBands
+		|| Settings.Bands.Num() - 1 != PrevCrossovers.Num()
+		|| Settings.bFourPole != bPrevFourPole);
 
 	// check if crossovers have changed
 	// if so, update before potentially initializing
@@ -54,8 +124,11 @@ void FSubmixEffectMultibandCompressor::OnPresetChanged()
 	{
 		// lighter way to update crossovers if nothing else needs to be reinit
 		BandSplitter.SetCrossovers(PrevCrossovers);
+		KeyBandSplitter.SetCrossovers(PrevCrossovers);
 	}
 
+	bBypass = Settings.bBypass;
+
 	Audio::EDynamicsProcessingMode::Type TypeToSet;
 	switch (Settings.DynamicsProcessorType)
 	{
@@ -98,11 +171,7 @@ void FSubmixEffectMultibandCompressor::OnPresetChanged()
 	{
 		Audio::FDynamicsProcessor& DynamicsProcessor = DynamicsProcessors[BandId];
 
-		Audio::EDynamicsProcessorChannelLinkMode ChannelLinkMode = Settings.bLinkChannels ?
-			Audio::EDynamicsProcessorChannelLinkMode::Peak
-			: Audio::EDynamicsProcessorChannelLinkMode::Disabled;
-
-		DynamicsProcessor.SetChannelLinkMode(ChannelLinkMode);
+		DynamicsProcessor.SetChannelLinkMode(static_cast<Audio::EDynamicsProcessorChannelLinkMode>(Settings.LinkMode));
 
 		DynamicsProcessor.SetLookaheadMsec(Settings.LookAheadMsec);
 		DynamicsProcessor.SetAnalogMode(Settings.bAnalogMode);
@@ -117,7 +186,104 @@ void FSubmixEffectMultibandCompressor::OnPresetChanged()
 
 		DynamicsProcessor.SetProcessingMode(TypeToSet);
 		DynamicsProcessor.SetPeakMode(PeakModeToSet);
+
+		// key audition will work fine if there's no external submix, and will simply bypass compression
+		DynamicsProcessor.SetKeyAudition(Settings.bKeyAudition);
+		DynamicsProcessor.SetKeyGain(Settings.KeyGainDb);
+
+		UpdateKeyFromSettings(Settings);
+	}
+}
+
+bool FSubmixEffectMultibandCompressor::UpdateKeySourcePatch()
+{
+	// Default (input as key) does not use source patch, so don't
+	// continue checking or updating state.
+	if (KeySource.GetType() == ESubmixEffectDynamicsKeySource::Default)
+	{
+		return false;
+	}
+
+	if (KeySource.Patch.IsValid())
+	{
+		return true;
+	}
+
+	switch (KeySource.GetType())
+	{
+		case ESubmixEffectDynamicsKeySource::AudioBus:
+		{
+			// Retrieving/mutating the MixerDevice is only safe during OnProcessAudio calls if
+			// it is not called during Teardown.  The DynamicsProcessor should be Reset via
+			// the OnDeviceDestroyed callback (prior to FAudioDevice::Teardown), so this call
+			// should never be hit during Teardown.
+			if (Audio::FMixerDevice* MixerDevice = GetMixerDevice())
+			{
+				const uint32 ObjectId = KeySource.GetObjectId();
+				if (ObjectId != INDEX_NONE)
+				{
+					KeySource.Patch = MixerDevice->AddPatchForAudioBus(ObjectId, 1.0f /* PatchGain */);
+					if (KeySource.Patch.IsValid())
+					{
+						for (Audio::FDynamicsProcessor& DynamicsProcessor : DynamicsProcessors)
+						{
+							DynamicsProcessor.SetKeyNumChannels(KeySource.GetNumChannels());
+						}
+
+						return true;
+					}
+					else if (KeySource.ShouldReportInactive())
+					{
+						KeySource.SetReportInactive(false);
+					}
+				}
+			}
+		}
+		break;
+
+		case ESubmixEffectDynamicsKeySource::Submix:
+		{
+			// Retrieving/mutating the MixerDevice is only safe during OnProcessAudio calls if
+			// it is not called during Teardown.  The DynamicsProcessor should be Reset via
+			// the OnDeviceDestroyed callback (prior to FAudioDevice::Teardown), so this call
+			// should never be hit during Teardown.
+			if (Audio::FMixerDevice* MixerDevice = GetMixerDevice())
+			{
+				const uint32 ObjectId = KeySource.GetObjectId();
+				if (ObjectId != INDEX_NONE)
+				{
+					KeySource.Patch = MixerDevice->AddPatchForSubmix(ObjectId, 1.0f /* PatchGain */);
+					if (KeySource.Patch.IsValid())
+					{
+						Audio::FMixerSubmixPtr SubmixPtr = MixerDevice->FindSubmixInstanceByObjectId(KeySource.GetObjectId());
+						if (SubmixPtr.IsValid())
+						{
+							const int32 SubmixNumChannels = SubmixPtr->GetNumOutputChannels();
+							KeySource.SetNumChannels(SubmixNumChannels);
+
+							for (Audio::FDynamicsProcessor& DynamicsProcessor : DynamicsProcessors)
+							{
+								DynamicsProcessor.SetKeyNumChannels(SubmixNumChannels);
+							}
+
+							return true;
+						}
+					}
+					else if (KeySource.ShouldReportInactive())
+					{
+						KeySource.SetReportInactive(false);
+					}
+				}
+			}
+		}
+		break;
+
+		case ESubmixEffectDynamicsKeySource::Default:
+		default:
+		break;
 	}
+
+	return false;
 }
 
 void FSubmixEffectMultibandCompressor::Initialize(FSubmixEffectMultibandCompressorSettings& Settings)
@@ -132,12 +298,16 @@ void FSubmixEffectMultibandCompressor::Initialize(FSubmixEffectMultibandCompress
 	}
 
 	Audio::EFilterOrder CrossoverMode = Settings.bFourPole ? Audio::EFilterOrder::FourPole : Audio::EFilterOrder::TwoPole;
-	BandSplitter.Init(NumChannels, SampleRate, CrossoverMode, PrevCrossovers);
+
+	BandSplitter.Init(NumChannels, SampleRate, CrossoverMode, true /*Phase Compensate*/, PrevCrossovers);
+	KeyBandSplitter.Init(NumChannels, SampleRate, CrossoverMode, false /*Phase Compensate*/, PrevCrossovers);
 
 	MultiBandBuffer.SetBands(NumBands);
+	KeyMultiBandBuffer.SetBands(NumBands);
 
 	DynamicsProcessors.Reset(NumBands);
 	DynamicsProcessors.AddDefaulted(NumBands);
+
 	for (int32 BandId = 0; BandId < NumBands; ++BandId)
 	{
 		DynamicsProcessors[BandId].Init(SampleRate, NumChannels);
@@ -162,44 +332,280 @@ void FSubmixEffectMultibandCompressor::OnProcessAudio(const FSoundEffectSubmixIn
 
 	const int32 NumSamples = InData.NumFrames * NumChannels;
 
-	if ((bInitialized == false) || (DynamicsProcessors.Num() == 0))
+	const Audio::FAlignedFloatBuffer& InBuffer = *InData.AudioBuffer;
+	Audio::FAlignedFloatBuffer& OutBuffer = *OutData.AudioBuffer;
+
+	if (bBypass || bInitialized == false || DynamicsProcessors.IsEmpty())
 	{
 		//passthrough
-		FMemory::Memcpy(OutData.AudioBuffer->GetData(), InData.AudioBuffer->GetData(), FrameSize * InData.NumFrames);
+		FMemory::Memcpy(OutBuffer.GetData(), InBuffer.GetData(), FrameSize * InData.NumFrames);
 		return;
 	}
 
-	const Audio::FAlignedFloatBuffer& InBuffer = *InData.AudioBuffer;
-	Audio::FAlignedFloatBuffer& OutBuffer = *OutData.AudioBuffer;
+	int32 NumKeyChannels = DynamicsProcessors[0].GetKeyNumChannels();
+	int32 NumKeySamples = InData.NumFrames * NumKeyChannels;
+
+	AudioExternal.Reset(NumKeySamples);
+
+	// If set to default, enforce num key channels to always be number of input channels.
+	// If either unset or KeySource was changed between frames back to 'Default', NumKeyChannels
+	// could be stale or left as initialized number of scratch channels.
+	if (KeySource.GetType() == ESubmixEffectDynamicsKeySource::Default)
+	{
+		if (InData.NumChannels != NumKeyChannels)
+		{
+			NumKeyChannels = InData.NumChannels;
+			NumKeySamples = InData.NumFrames * NumKeyChannels;
+
+			for (Audio::FDynamicsProcessor& DynamicsProcessor : DynamicsProcessors)
+			{
+				DynamicsProcessor.SetKeyNumChannels(NumKeyChannels);
+			}
+		}
+	}
+	else
+	{
+		AudioExternal.AddZeroed(NumKeySamples);
+	}
+
+	const bool bUseKey = UpdateKeySourcePatch();
+
+	if (bUseKey)
+	{
+		KeySource.Patch->PopAudio(AudioExternal.GetData(), NumKeySamples, true /* bUseLatestAudio */);
+	}
 
+	if (InData.NumChannels != DynamicsProcessors[0].GetNumChannels())
+	{
+		for (Audio::FDynamicsProcessor& DynamicsProcessor : DynamicsProcessors)
+		{
+			DynamicsProcessor.SetNumChannels(InData.NumChannels);
+		}
+	}
+
+	// Zero buffer so multiple bands can all sum to it
 	FMemory::Memzero(OutBuffer.GetData(), FrameSize * InData.NumFrames);
 
 	Audio::FStackSampleBuffer ScratchBuffer;
+
+	// Scratch buffer is stack allocated with a max size. If passed more samples than it can hold, process in blocks
 	const int32 BlockSize = FMath::Min(ScratchBuffer.Max(), NumSamples);
 
 	if (BlockSize > MultiBandBuffer.NumSamples)
 	{
 		MultiBandBuffer.SetSamples(NumSamples);
+
+		if (bUseKey)
+		{
+			KeyMultiBandBuffer.SetSamples(NumKeySamples);
+		}
 	}
 
 	for (int32 SampleIdx = 0; SampleIdx < NumSamples; SampleIdx += BlockSize)
 	{
 		const float* InPtr = &InBuffer[SampleIdx];
 		float* OutPtr = &OutBuffer[SampleIdx];
+		const int32 NumBlockFrames = BlockSize / NumChannels;
+
+		const float* KeyPtr = nullptr;
+		
+		if (bUseKey)
+		{
+			const int32 KeySampleIdx = (SampleIdx / NumChannels) * NumKeyChannels;
+			ensure(KeySampleIdx < AudioExternal.Num());
+			KeyPtr = &AudioExternal[KeySampleIdx];
+		}
 
 		ScratchBuffer.SetNumUninitialized(BlockSize);
 
-		BandSplitter.ProcessAudioBuffer(InPtr, MultiBandBuffer, BlockSize / NumChannels);
+		BandSplitter.ProcessAudioBuffer(InPtr, MultiBandBuffer, NumBlockFrames);
+
+		if (bUseKey)
+		{
+			KeyBandSplitter.ProcessAudioBuffer(KeyPtr, KeyMultiBandBuffer, NumBlockFrames);
+		}
+		else
+		{
+			// reset key gain to unity in case this has changed recently on another thread
+			for (Audio::FDynamicsProcessor& DynamicsProcessor : DynamicsProcessors)
+			{
+				//DynamicsProcessor.SetKeyGain(0.0f);
+			}
+		}
 
 		for (int32 Band = 0; Band < DynamicsProcessors.Num(); ++Band)
 		{
-			DynamicsProcessors[Band].ProcessAudio(MultiBandBuffer[Band], BlockSize, ScratchBuffer.GetData());
+			DynamicsProcessors[Band].ProcessAudio(MultiBandBuffer[Band], BlockSize, ScratchBuffer.GetData(), bUseKey ? KeyMultiBandBuffer[Band] : nullptr);
 			Audio::MixInBufferFast(ScratchBuffer.GetData(), OutPtr, BlockSize);
 		}
 	}
 }
 
+void FSubmixEffectMultibandCompressor::UpdateKeyFromSettings(const FSubmixEffectMultibandCompressorSettings& InSettings)
+{
+	uint32 ObjectId = INDEX_NONE;
+	int32 SourceNumChannels = 0;
+	switch (InSettings.KeySource)
+	{
+		case ESubmixEffectDynamicsKeySource::AudioBus:
+		{
+			if (InSettings.ExternalAudioBus)
+			{
+				ObjectId = InSettings.ExternalAudioBus->GetUniqueID();
+				SourceNumChannels = static_cast<int32>(InSettings.ExternalAudioBus->AudioBusChannels) + 1;
+			}
+		}
+		break;
+
+		case ESubmixEffectDynamicsKeySource::Submix:
+		{
+			if (InSettings.ExternalSubmix)
+			{
+				ObjectId = InSettings.ExternalSubmix->GetUniqueID();
+			}
+		}
+		break;
+
+		default:
+		{
+		}
+		break;
+	}
+
+	KeySource.Update(InSettings.KeySource, ObjectId, SourceNumChannels);
+}
+
+void FSubmixEffectMultibandCompressor::OnDeviceCreated(Audio::FDeviceId InDeviceId)
+{
+	if (InDeviceId == DeviceId)
+	{
+		GET_EFFECT_SETTINGS(SubmixEffectMultibandCompressor);
+		UpdateKeyFromSettings(Settings);
+
+		FAudioDeviceManagerDelegates::OnAudioDeviceCreated.Remove(DeviceCreatedHandle);
+	}
+}
+
+void FSubmixEffectMultibandCompressor::OnDeviceDestroyed(Audio::FDeviceId InDeviceId)
+{
+	if (InDeviceId == DeviceId)
+	{
+		// Reset the key on device destruction to avoid reinitializing
+		// it during FAudioDevice::Teardown via ProcessAudio.
+		ResetKey();
+		FAudioDeviceManagerDelegates::OnAudioDeviceDestroyed.Remove(DeviceDestroyedHandle);
+	}
+}
+
+/**
+ * Preset
+ */
+void USubmixEffectMultibandCompressorPreset::OnInit()
+{
+	switch (Settings.KeySource)
+	{
+		case ESubmixEffectDynamicsKeySource::AudioBus:
+		{
+			SetAudioBus(Settings.ExternalAudioBus);
+		}
+		break;
+
+		case ESubmixEffectDynamicsKeySource::Submix:
+		{
+			SetExternalSubmix(Settings.ExternalSubmix);
+		}
+		break;
+
+		default:
+		break;
+	}
+}
+
+#if WITH_EDITOR
+void USubmixEffectMultibandCompressorPreset::PostEditChangeChainProperty(struct FPropertyChangedChainEvent& InChainEvent)
+{
+	if (InChainEvent.GetPropertyName() == GET_MEMBER_NAME_CHECKED(FSubmixEffectMultibandCompressorSettings, KeySource))
+	{
+		switch (Settings.KeySource)
+		{
+			case ESubmixEffectDynamicsKeySource::AudioBus:
+			{
+				Settings.ExternalSubmix = nullptr;
+			}
+			break;
+
+			case ESubmixEffectDynamicsKeySource::Submix:
+			{
+				Settings.ExternalAudioBus = nullptr;
+			}
+			break;
+
+			case ESubmixEffectDynamicsKeySource::Default:
+			default:
+			{
+				Settings.ExternalSubmix = nullptr;
+				Settings.ExternalAudioBus = nullptr;
+				static_assert(static_cast<int32>(ESubmixEffectDynamicsKeySource::Count) == 3, "Possible missing KeySource switch case coverage");
+			}
+			break;
+		}
+	}
+
+	Super::PostEditChangeChainProperty(InChainEvent);
+}
+#endif // WITH_EDITOR
+
+void USubmixEffectMultibandCompressorPreset::ResetKey()
+{
+	EffectCommand<FSubmixEffectMultibandCompressor>([](FSubmixEffectMultibandCompressor& Instance)
+	{
+		Instance.ResetKey();
+	});
+}
+
+void USubmixEffectMultibandCompressorPreset::SetAudioBus(UAudioBus* InAudioBus)
+{
+	int32 BusChannels = 0;
+	if (InAudioBus)
+	{
+		BusChannels = static_cast<int32>(InAudioBus->AudioBusChannels) + 1;
+		SetKey(ESubmixEffectDynamicsKeySource::AudioBus, InAudioBus, BusChannels);
+	}
+	else
+	{
+		ResetKey();
+	}
+}
+
+void USubmixEffectMultibandCompressorPreset::SetExternalSubmix(USoundSubmix* InSubmix)
+{
+	if (InSubmix)
+	{
+		SetKey(ESubmixEffectDynamicsKeySource::Submix, InSubmix);
+	}
+	else
+	{
+		ResetKey();
+	}
+}
+
+void USubmixEffectMultibandCompressorPreset::SetKey(ESubmixEffectDynamicsKeySource InKeySource, UObject* InObject, int32 InNumChannels)
+{
+	if (InObject)
+	{
+		EffectCommand<FSubmixEffectMultibandCompressor>([this, ObjectId = InObject->GetUniqueID(), InKeySource, InNumChannels](FSubmixEffectMultibandCompressor& Instance)
+		{
+			Instance.KeySource.Update(InKeySource, ObjectId, InNumChannels);
+		});
+	}
+}
+
 void USubmixEffectMultibandCompressorPreset::SetSettings(const FSubmixEffectMultibandCompressorSettings& InSettings)
 {
 	UpdateSettings(InSettings);
+
+	IterateEffects<FSubmixEffectMultibandCompressor>([&](FSubmixEffectMultibandCompressor& Instance)
+	{
+		Instance.UpdateKeyFromSettings(InSettings);
+	});
 }
diff --git a/Engine/Source/Runtime/SignalProcessing/Private/LinkwitzRileyBandSplitter.cpp b/Engine/Source/Runtime/SignalProcessing/Private/LinkwitzRileyBandSplitter.cpp
--- a/Engine/Source/Runtime/SignalProcessing/Private/LinkwitzRileyBandSplitter.cpp
+++ b/Engine/Source/Runtime/SignalProcessing/Private/LinkwitzRileyBandSplitter.cpp
@@ -8,7 +8,7 @@
 #endif
 namespace Audio
 {
-	void FLinkwitzRileyBandSplitter::Init(const int32 InChannels, const float InSampleRate, const EFilterOrder InFilterOrder, const TArray<float>& InCrossovers)
+	void FLinkwitzRileyBandSplitter::Init(const int32 InChannels, const float InSampleRate, const EFilterOrder InFilterOrder, const bool bInPhaseCompensate, const TArray<float>& InCrossovers)
 	{
 		NumBands = InCrossovers.Num() + 1;
 		NumChannels = InChannels;
@@ -42,7 +42,8 @@ namespace Audio
 
 		for (int32 BandId = 1; BandId < NumBands; BandId++)
 		{
-			BandFilters[BandId].Filters.AddDefaulted(NumBands - BandId);
+			const int32 NumFilters = NumBands - BandId;
+			BandFilters[BandId].Filters.AddDefaulted( bInPhaseCompensate ? NumFilters : FMath::Clamp(NumFilters, 1, 2) );
 		}
 
 		// band 0 special case
diff --git a/Engine/Source/Runtime/SignalProcessing/Public/DSP/LinkwitzRileyBandSplitter.h b/Engine/Source/Runtime/SignalProcessing/Public/DSP/LinkwitzRileyBandSplitter.h
--- a/Engine/Source/Runtime/SignalProcessing/Public/DSP/LinkwitzRileyBandSplitter.h
+++ b/Engine/Source/Runtime/SignalProcessing/Public/DSP/LinkwitzRileyBandSplitter.h
@@ -84,7 +84,8 @@ namespace Audio
 		// initalize filters
 		void Init(const int32 InChannels,
 				  const float InSampleRate,
-				  const EFilterOrder FilterOrder,
+				  const EFilterOrder FilterOrder, 
+				  const bool bInPhaseCompensate,
 				  const TArray<float>& InCrossovers); // Always InBands - 1 Crossovers
 
 		void ProcessAudioFrame(const float* InBuffer, FMultibandBuffer& OutBuffer);
