Audio: Initialized some uninitialized struct members to fix a startup warning (soon to be an error)

diff --git a/Engine/Plugins/Runtime/Synthesis/Source/Synthesis/Classes/SubmixEffects/SubmixEffectMultiBandCompressor.h b/Engine/Plugins/Runtime/Synthesis/Source/Synthesis/Classes/SubmixEffects/SubmixEffectMultiBandCompressor.h
--- a/Engine/Plugins/Runtime/Synthesis/Source/Synthesis/Classes/SubmixEffects/SubmixEffectMultiBandCompressor.h
+++ b/Engine/Plugins/Runtime/Synthesis/Source/Synthesis/Classes/SubmixEffects/SubmixEffectMultiBandCompressor.h
@@ -78,7 +78,7 @@ struct SYNTHESIS_API FSubmixEffectMultibandCompressorSettings
 
 	// Whether or not to bypass effect
 	UPROPERTY(EditAnywhere, BlueprintReadWrite, Category = General, meta = (DisplayName = "Bypass", DisplayAfter = "DynamicsProcessorType"))
-	bool bBypass;
+	bool bBypass = false;
 
 	UPROPERTY(EditAnywhere, BlueprintReadWrite, Category = Sidechain, meta = (EditCondition = "!bBypass"))
 	ESubmixEffectDynamicsKeySource KeySource = ESubmixEffectDynamicsKeySource::Default;
@@ -93,11 +93,11 @@ struct SYNTHESIS_API FSubmixEffectMultibandCompressorSettings
 
 	// Gain to apply to key signal if external input is supplied
 	UPROPERTY(EditAnywhere, BlueprintReadWrite, Category = Sidechain, meta = (DisplayName = "External Input Gain (dB)", EditCondition = "ExternalSubmix != nullptr || ExternalAudioBus != nullptr", UIMin = "-60.0", UIMax = "30.0"))
-	float KeyGainDb;
+	float KeyGainDb = 0.0f;
 
 	// Audition the key modulation signal, bypassing enveloping and processing the input signal.
 	UPROPERTY(EditAnywhere, BlueprintReadWrite, Category = Sidechain, meta = (DisplayName = "Key Audition", EditCondition = "!bBypass"))
-	bool bKeyAudition;
+	bool bKeyAudition = false;
 
 	// Each band is a full dynamics processor, affecting at a unique frequency range
 	UPROPERTY(EditAnywhere, BlueprintReadWrite, Category = Bands)
