Fix for init delay length on simple distance delay

diff --git a/Engine/Plugins/Runtime/Synthesis/Source/Synthesis/Classes/SourceEffects/SourceEffectSimpleDelay.h b/Engine/Plugins/Runtime/Synthesis/Source/Synthesis/Classes/SourceEffects/SourceEffectSimpleDelay.h
--- a/Engine/Plugins/Runtime/Synthesis/Source/Synthesis/Classes/SourceEffects/SourceEffectSimpleDelay.h
+++ b/Engine/Plugins/Runtime/Synthesis/Source/Synthesis/Classes/SourceEffects/SourceEffectSimpleDelay.h
@@ -62,6 +62,7 @@ protected:
 	TArray<Audio::FDelay> Delays;
 	TArray<float> FeedbackSamples;
 	FSourceEffectSimpleDelaySettings SettingsCopy;
+	bool bIsInit = true;
 };
 
 UCLASS(ClassGroup = AudioSourceEffect, meta = (BlueprintSpawnableComponent))
diff --git a/Engine/Plugins/Runtime/Synthesis/Source/Synthesis/Private/SourceEffectSimpleDelay.cpp b/Engine/Plugins/Runtime/Synthesis/Source/Synthesis/Private/SourceEffectSimpleDelay.cpp
--- a/Engine/Plugins/Runtime/Synthesis/Source/Synthesis/Private/SourceEffectSimpleDelay.cpp
+++ b/Engine/Plugins/Runtime/Synthesis/Source/Synthesis/Private/SourceEffectSimpleDelay.cpp
@@ -5,6 +5,7 @@
 void FSourceEffectSimpleDelay::Init(const FSoundEffectSourceInitData& InitData)
 {
 	bIsActive = true;
+	bIsInit = true;
 
 	FeedbackSamples.AddZeroed(InitData.NumSourceChannels);
 	Delays.AddDefaulted(InitData.NumSourceChannels);
@@ -26,8 +27,9 @@ void FSourceEffectSimpleDelay::OnPresetChanged()
 	{
 		for (Audio::FDelay& Delay : Delays)
 		{
-			Delay.SetEasedDelayMsec(SettingsCopy.DelayAmount * 1000.0f);
+			Delay.SetEasedDelayMsec(SettingsCopy.DelayAmount * 1000.0f, bIsInit);
 		}
+		bIsInit = false;
 	}
 }
 
@@ -40,8 +42,9 @@ void FSourceEffectSimpleDelay::ProcessAudio(const FSoundEffectSourceInputData& I
 
 		for (Audio::FDelay& Delay : Delays)
 		{
-			Delay.SetEasedDelayMsec(DelayAmountMsec);
+			Delay.SetEasedDelayMsec(DelayAmountMsec, bIsInit);
 		}
+		bIsInit = false;
 	}
 
 	int32 NumChannels = Delays.Num();
