Fix tooltip for envelope follower

diff --git a/Engine/Plugins/Runtime/Synthesis/Source/Synthesis/Classes/SourceEffects/SourceEffectEnvelopeFollower.h b/Engine/Plugins/Runtime/Synthesis/Source/Synthesis/Classes/SourceEffects/SourceEffectEnvelopeFollower.h
--- a/Engine/Plugins/Runtime/Synthesis/Source/Synthesis/Classes/SourceEffects/SourceEffectEnvelopeFollower.h
+++ b/Engine/Plugins/Runtime/Synthesis/Source/Synthesis/Classes/SourceEffects/SourceEffectEnvelopeFollower.h
@@ -136,11 +136,11 @@ public:
 	UFUNCTION(BlueprintCallable, Category = "Audio|Effects")
 	void SetSettings(const FSourceEffectEnvelopeFollowerSettings& InSettings);
 
-	/** Adds a submix effect preset to the master submix. */
+	/** Registers an envelope follower listener with the effect. */
 	UFUNCTION(BlueprintCallable, Category = "Audio|Effects", meta = (WorldContext = "WorldContextObject"))
 	void RegisterEnvelopeFollowerListener(UEnvelopeFollowerListener* EnvelopeFollowerListener);
 
-	/** Adds a submix effect preset to the master submix. */
+	/** Unregisters an envelope follower listener with the effect. */
 	UFUNCTION(BlueprintCallable, Category = "Audio|Effects", meta = (WorldContext = "WorldContextObject"))
 	void UnregisterEnvelopeFollowerListener(UEnvelopeFollowerListener* EnvelopeFollowerListener);
 
