Fixing mispelling of chord to cord

diff --git a/Engine/Plugins/Runtime/Synthesis/Source/Synthesis/Classes/SynthComponents/EpicSynth1Component.h b/Engine/Plugins/Runtime/Synthesis/Source/Synthesis/Classes/SynthComponents/EpicSynth1Component.h
--- a/Engine/Plugins/Runtime/Synthesis/Source/Synthesis/Classes/SynthComponents/EpicSynth1Component.h
+++ b/Engine/Plugins/Runtime/Synthesis/Source/Synthesis/Classes/SynthComponents/EpicSynth1Component.h
@@ -268,7 +268,7 @@ struct SYNTHESIS_API FModularSynthPreset : public FTableRowBase
 	UPROPERTY(EditAnywhere, BlueprintReadWrite, Category = "Synth|Preset", meta = (ClampMin = "0.0", ClampMax = "20.0", UIMin = "0.0", UIMax = "20.0"))
 	float ChorusFrequency;
 
-	// The modular synth patch chords to use for the synth. Allows routing the LFO1/LFO2 and Modulation Envelope to any patchable destination. 
+	// The modular synth patch cords to use for the synth. Allows routing the LFO1/LFO2 and Modulation Envelope to any patchable destination. 
 	UPROPERTY(EditAnywhere, BlueprintReadWrite, Category = "Synth|Preset")
 	TArray<FEpicSynth1Patch> Patches;
 
diff --git a/Engine/Source/Runtime/Engine/Classes/Sound/AudioBus.h b/Engine/Source/Runtime/Engine/Classes/Sound/AudioBus.h
--- a/Engine/Source/Runtime/Engine/Classes/Sound/AudioBus.h
+++ b/Engine/Source/Runtime/Engine/Classes/Sound/AudioBus.h
@@ -22,7 +22,7 @@ enum class EAudioBusChannels : uint8
 // Function to retrieve an audio bus buffer given a handle
 // static float* GetAudioBusBuffer(const FAudioBusHandle& AudioBusHandle);
 
-// An audio bus is an object which represents an audio patch chord. Audio can be sent to it. It can be sonified using USoundSourceBuses.
+// An audio bus is an object which represents an audio patch cord. Audio can be sent to it. It can be sonified using USoundSourceBuses.
 // Instances of the audio bus are created in the audio engine. 
 UCLASS(ClassGroup = Sound, meta = (BlueprintSpawnableComponent))
 class ENGINE_API UAudioBus : public UObject
