Make convolution reverb available as a source effect

diff --git a/Engine/Plugins/Runtime/Synthesis/Source/Synthesis/Classes/SourceEffects/SourceEffectConvolutionReverb.h b/Engine/Plugins/Runtime/Synthesis/Source/Synthesis/Classes/SourceEffects/SourceEffectConvolutionReverb.h
new file mode 100644
--- /dev/null
+++ b/Engine/Plugins/Runtime/Synthesis/Source/Synthesis/Classes/SourceEffects/SourceEffectConvolutionReverb.h
@@ -0,0 +1,170 @@
+// Copyright Epic Games, Inc. All Rights Reserved.
+#pragma once
+
+#include "CoreMinimal.h"
+#include "UObject/Object.h"
+#include "DSP/Dsp.h"
+#include "EffectConvolutionReverb.h"
+#include "Sound/SoundEffectSource.h"
+#include "SourceEffectConvolutionReverb.generated.h"
+
+USTRUCT(BlueprintType)
+struct SYNTHESIS_API FSourceEffectConvolutionReverbSettings
+{
+	GENERATED_USTRUCT_BODY()
+
+	FSourceEffectConvolutionReverbSettings();
+
+	/* Used to account for energy added by convolution with "loud" Impulse Responses. 
+	 * This value is not directly editable in the editor because it is copied from the 
+	 * associated UAudioImpulseResponse. */
+	UPROPERTY();
+	float NormalizationVolumeDb;
+
+	// Controls how much of the wet signal is mixed into the output, in Decibels
+	UPROPERTY(EditAnywhere, BlueprintReadWrite, Category = SubmixEffectPreset, meta = (ClampMin = "-96.0", ClampMax = "0.0"));
+	float WetVolumeDb;
+
+	// Controls how much of the dry signal is mixed into the output, in Decibels
+	UPROPERTY(EditAnywhere, BlueprintReadWrite, Category = SubmixEffectPreset, meta = (ClampMin = "-96.0", ClampMax = "0.0"));
+	float DryVolumeDb;
+
+	/* If true, input audio is directly routed to output audio with applying any effect. */
+	UPROPERTY(EditAnywhere, BlueprintReadWrite, Category = SourceEffectPreset)
+	bool bBypass;
+};
+
+
+/** Audio render thread effect object. */
+class SYNTHESIS_API FSourceEffectConvolutionReverb : public FSoundEffectSource
+{
+public:
+	FSourceEffectConvolutionReverb() = delete;
+	// Construct a convolution object with an existing preset. 
+	FSourceEffectConvolutionReverb(const USourceEffectConvolutionReverbPreset* InPreset);
+
+	~FSourceEffectConvolutionReverb();
+
+	// Called on an audio effect at initialization on main thread before audio processing begins.
+	virtual void Init(const FSoundEffectSourceInitData& InInitData) override;
+
+	// Called when an audio effect preset settings is changed.
+	virtual void OnPresetChanged() override;
+
+	// Process the input block of audio. Called on audio thread.
+	virtual void ProcessAudio(const FSoundEffectSourceInputData& InData, float* OutAudioBufferData) override;
+
+	// Call on the game thread in order to update the impulse response and hardware acceleration
+	// used in this effect.
+	AudioConvReverbIntrinsics::FVersionData UpdateConvolutionReverb(const USourceEffectConvolutionReverbPreset* InPreset);
+
+	void RebuildConvolutionReverb();
+private:
+	// Sets current runtime settings for convolution reverb which do *not* trigger
+	// a FConvolutionReverb rebuild.  These settings will be applied to FConvolutionReverb 
+	// at the next call to UpdateParameters()
+	void SetConvolutionReverbParameters(const FSourceEffectConvolutionReverbSettings& InSettings);
+
+	// Reverb performs majority of DSP operations
+	TSharedRef<Audio::FEffectConvolutionReverb> Reverb;
+
+	float WetVolume = 1.f;
+	float DryVolume = 1.f;
+	int32 NumChannels = 0;
+};
+
+UCLASS()
+class SYNTHESIS_API USourceEffectConvolutionReverbPreset : public USoundEffectSourcePreset
+{
+	GENERATED_BODY()
+
+public:
+	
+	USourceEffectConvolutionReverbPreset(const FObjectInitializer& ObjectInitializer);
+
+	/* Note: The SourceEffect boilerplate macros could not be utilized here because
+	 * the "CreateNewEffect" implementation differed from those available in the
+	 * boilerplate macro.
+	 */
+
+	virtual bool CanFilter() const override;
+
+	virtual bool HasAssetActions() const;
+
+	virtual FText GetAssetActionName() const override;
+
+	virtual UClass* GetSupportedClass() const override;
+
+	virtual FSoundEffectBase* CreateNewEffect() const override;
+
+	virtual USoundEffectPreset* CreateNewPreset(UObject* InParent, FName Name, EObjectFlags Flags) const override;
+
+	virtual void Init() override;
+
+
+	FSourceEffectConvolutionReverbSettings GetSettings() const;
+
+	/** Set the convolution reverb settings */
+	UFUNCTION(BlueprintCallable, Category = "Audio|Effects")
+	void SetSettings(const FSourceEffectConvolutionReverbSettings& InSettings);
+
+	/** Set the convolution reverb impulse response */
+	UFUNCTION(BlueprintCallable, BlueprintSetter, Category = "Audio|Effects")
+	void SetImpulseResponse(UAudioImpulseResponse* InImpulseResponse);
+
+	/** The impulse response used for convolution. */
+	UPROPERTY(EditAnywhere, BlueprintReadWrite, BlueprintSetter = SetImpulseResponse, Category = SourceEffectPreset)
+	TObjectPtr<UAudioImpulseResponse> ImpulseResponse;
+
+	/** ConvolutionReverbPreset Preset Settings. */
+	UPROPERTY(EditAnywhere, BlueprintReadWrite, BlueprintSetter = SetSettings, Category = SourceEffectPreset)
+	FSourceEffectConvolutionReverbSettings Settings;
+
+
+	/** Set the internal block size. This can effect latency and performance. Higher values will result in
+	 * lower CPU costs while lower values will result higher CPU costs. Latency may be affected depending
+	 * on the interplay between audio engines buffer sizes and this effects block size. Generally, higher
+	 * values result in higher latency, and lower values result in lower latency. 
+	 */
+	UPROPERTY(EditAnywhere, BlueprintReadOnly, AdvancedDisplay, Category = SourceEffectPreset)
+	ESubmixEffectConvolutionReverbBlockSize BlockSize;
+
+	/** Opt into hardware acceleration of the convolution reverb (if available) */
+	UPROPERTY(EditAnywhere, BlueprintReadOnly, AdvancedDisplay, Category = SourceEffectPreset)
+	bool bEnableHardwareAcceleration;
+	
+
+#if WITH_EDITORONLY_DATA
+	// Binds to the UAudioImpulseRespont::OnObjectPropertyChanged delegate of the current ImpulseResponse
+	void BindToImpulseResponseObjectChange();
+
+	virtual void PreEditChange(FProperty* PropertyAboutToChange) override;
+	virtual void PostEditChangeProperty(FPropertyChangedEvent& PropertyChangedEvent) override;
+
+	// Called when a property changes on the ImpulseResponse object
+	void PostEditChangeImpulseProperty(FPropertyChangedEvent& PropertyChangedEvent);
+#endif
+
+	virtual void PostLoad() override;
+
+private:
+	void SetImpulseResponseSettings(UAudioImpulseResponse* InImpulseResponse);
+
+	void UpdateSettings();
+
+	void UpdateDeprecatedProperties();
+
+	// This method requires that the effect is registered with a preset.  If this 
+	// effect is not registered with a preset, then this will not update the convolution
+	// algorithm.
+	void RebuildConvolutionReverb();
+
+	mutable FCriticalSection SettingsCritSect; 
+	FSourceEffectConvolutionReverbSettings SettingsCopy; 
+
+#if WITH_EDITORONLY_DATA
+
+	TMap<UObject*, FDelegateHandle> DelegateHandles;
+#endif
+};
+
diff --git a/Engine/Plugins/Runtime/Synthesis/Source/Synthesis/Classes/SubmixEffects/SubmixEffectConvolutionReverb.h b/Engine/Plugins/Runtime/Synthesis/Source/Synthesis/Classes/SubmixEffects/SubmixEffectConvolutionReverb.h
--- a/Engine/Plugins/Runtime/Synthesis/Source/Synthesis/Classes/SubmixEffects/SubmixEffectConvolutionReverb.h
+++ b/Engine/Plugins/Runtime/Synthesis/Source/Synthesis/Classes/SubmixEffects/SubmixEffectConvolutionReverb.h
@@ -3,73 +3,11 @@
 
 #include "CoreMinimal.h"
 #include "UObject/Object.h"
-#include "ConvolutionReverb.h"
 #include "DSP/Dsp.h"
-#include "DSP/ConvolutionAlgorithm.h"
+#include "EffectConvolutionReverb.h"
 #include "Sound/SoundEffectSubmix.h"
 #include "SubmixEffectConvolutionReverb.generated.h"
 
-/** Block size of convolution reverb algorithm. */
-UENUM()
-enum class ESubmixEffectConvolutionReverbBlockSize : uint8
-{
-	/** 256 audio frames per a block. */
-	BlockSize256,
-
-	/** 512 audio frames per a block. */
-	BlockSize512,
-
-	/** 1024 audio frames per a block. */
-	BlockSize1024
-};
-
-#if WITH_EDITORONLY_DATA
-DECLARE_MULTICAST_DELEGATE_OneParam(FAudioImpulseResponsePropertyChange, FPropertyChangedEvent&);
-#endif
-
-// ========================================================================
-// UAudioImpulseResponse
-// UAsset used to represent Imported Impulse Responses
-// ========================================================================
-UCLASS(BlueprintType)
-class SYNTHESIS_API UAudioImpulseResponse : public UObject
-{
-	GENERATED_BODY()
-
-public:
-	UAudioImpulseResponse();
-	/** The interleaved audio samples used in convolution. */
-	UPROPERTY()
-	TArray<float> ImpulseResponse;
-
-	/** The number of channels in impulse response. */
-	UPROPERTY()
-	int32 NumChannels;
-
-	/** The original sample rate of the impulse response. */
-	UPROPERTY()
-	int32 SampleRate;
-
-	/* Used to account for energy added by convolution with "loud" Impulse Responses. */
-	UPROPERTY(EditAnywhere, BlueprintReadOnly, Category = SubmixEffectPreset, meta = (ClamMin = "-60.0", ClampMax = "15.0"))
-	float NormalizationVolumeDb;
-
-	/* If true, impulse response channels are interpreted as true stereo which allows channel crosstalk. If false, impulse response channels are interpreted as independent channel impulses. */
-	UPROPERTY(EditAnywhere, BlueprintReadOnly, Category = SubmixEffectPreset, meta = (EditCondition = "(NumChannels > 0) && (NumChannels % 2 == 0)"))
-	bool bTrueStereo;
-
-	UPROPERTY(meta = ( DeprecatedProperty ) )
-	TArray<float> IRData_DEPRECATED;
-
-#if WITH_EDITORONLY_DATA
-	virtual void PostEditChangeProperty(FPropertyChangedEvent& PropertyChangedEvent) override;
-
-	// This delegate is called whenever PostEditChangeProperty called.
-	FAudioImpulseResponsePropertyChange OnObjectPropertyChanged;
-#endif
-};
-
-
 USTRUCT(BlueprintType)
 struct SYNTHESIS_API FSubmixEffectConvolutionReverbSettings
 {
@@ -77,13 +15,20 @@ struct SYNTHESIS_API FSubmixEffectConvolutionReverbSettings
 
 	FSubmixEffectConvolutionReverbSettings();
 
-
 	/* Used to account for energy added by convolution with "loud" Impulse Responses. 
 	 * This value is not directly editable in the editor because it is copied from the 
 	 * associated UAudioImpulseResponse. */
 	UPROPERTY();
 	float NormalizationVolumeDb;
 
+	// Controls how much of the wet signal is mixed into the output, in Decibels
+	UPROPERTY(EditAnywhere, BlueprintReadWrite, Category = SubmixEffectPreset, meta = (ClampMin = "-96.0", ClampMax = "0.0"));
+	float WetVolumeDb;
+
+	// Controls how much of the dry signal is mixed into the output, in Decibels
+	UPROPERTY(EditAnywhere, BlueprintReadWrite, Category = SubmixEffectPreset, meta = (ClampMin = "-96.0", ClampMax = "0.0"));
+	float DryVolumeDb;
+
 	/* If true, input audio is directly routed to output audio with applying any effect. */
 	UPROPERTY(EditAnywhere, BlueprintReadWrite, Category = SubmixEffectPreset)
 	bool bBypass;
@@ -122,7 +67,6 @@ struct SYNTHESIS_API FSubmixEffectConvolutionReverbSettings
 
 	UPROPERTY(meta = ( DeprecatedProperty ) )
 	bool AllowHardwareAcceleration_DEPRECATED;
-
 };
 
 
@@ -130,18 +74,6 @@ struct SYNTHESIS_API FSubmixEffectConvolutionReverbSettings
 class SYNTHESIS_API FSubmixEffectConvolutionReverb : public FSoundEffectSubmix
 {
 public:
-
-	typedef int32 FConvolutionReverbID;
-
-	// Data used to track versions of current running convolution.
-	struct FVersionData
-	{
-		FConvolutionReverbID ConvolutionID;
-
-		FVersionData();
-		bool operator==(const FVersionData& Other) const;
-	};
-
 	
 	FSubmixEffectConvolutionReverb() = delete;
 	// Construct a convolution object with an existing preset. 
@@ -160,16 +92,12 @@ public:
 
 	// Call on the game thread in order to update the impulse response and hardware acceleration
 	// used in this submix effect.
-	FVersionData UpdateConvolutionReverb(const USubmixEffectConvolutionReverbPreset* InPreset);
-
-	// Sets the convolution algorithm object used if the version data matches the most up-to-date
-	// version data tracked within this object. This method must be called on the audio render thread.
-	void SetConvolutionReverbIfCurrent(TUniquePtr<Audio::FConvolutionReverb> InReverb, const FVersionData& InVersionData);
+	AudioConvReverbIntrinsics::FVersionData UpdateConvolutionReverb(const USubmixEffectConvolutionReverbPreset* InPreset);
 
-	// Returns the data needed to build a convolution reverb
-	Audio::FConvolutionReverbInitData CreateConvolutionReverbInitData();
+	void RebuildConvolutionReverb();
 
-	Audio::FConvolutionReverbSettings GetParameters() const;
+	float WetVolume = 1.f;
+	float DryVolume = 1.f;
 
 private:
 	// Sets current runtime settings for convolution reverb which do *not* trigger
@@ -177,35 +105,8 @@ private:
 	// at the next call to UpdateParameters()
 	void SetConvolutionReverbParameters(const FSubmixEffectConvolutionReverbSettings& InSettings);
 
-	// Increments the internal verion number and returns a copy of the
-	// latest version information.
-	FVersionData UpdateVersion();
-
-	// Updates parameters on the render thread.
-	void UpdateParameters();
-
-	// Params object for copying between audio render thread and outside threads.
-	Audio::TParams<Audio::FConvolutionReverbSettings> Params;
-
-	// ConvolutionReverb performs majority of DSP operations
-	TUniquePtr<Audio::FConvolutionReverb> ConvolutionReverb;
-
-	// Audio sample rate
-	float SampleRate;
-
-	// Number of input and output channels. These are TAtomic
-	// since they can be accessed on multiple threads.
-	TAtomic<int32> NumInputChannels;
-	TAtomic<int32> NumOutputChannels;
-
-	TAtomic<bool> bBypass;
-
-	mutable FCriticalSection VersionDataCriticalSection;
-	FVersionData VersionData;
-
-	mutable FCriticalSection ConvReverbInitDataCriticalSection;
-	// Internal copy of data needed to create FConvolutionReverb
-	Audio::FConvolutionReverbInitData ConvReverbInitData;
+	// Reverb performs majority of DSP operations
+	TSharedRef<Audio::FEffectConvolutionReverb> Reverb;
 };
 
 UCLASS()
diff --git a/Engine/Plugins/Runtime/Synthesis/Source/Synthesis/Private/ConvolutionReverb.cpp b/Engine/Plugins/Runtime/Synthesis/Source/Synthesis/Private/ConvolutionReverb.cpp
--- a/Engine/Plugins/Runtime/Synthesis/Source/Synthesis/Private/ConvolutionReverb.cpp
+++ b/Engine/Plugins/Runtime/Synthesis/Source/Synthesis/Private/ConvolutionReverb.cpp
@@ -455,35 +455,38 @@ namespace Audio
 		return Settings;
 	}
 
+	void FConvolutionReverb::ProcessAudio(int32 InNumInputChannels, const FAlignedFloatBuffer& InputAudio, int32 InNumOutputChannels, FAlignedFloatBuffer& OutputAudio)
+	{
+		check(InNumInputChannels != 0);
+		check(InputAudio.Num() % InNumInputChannels == 0);
+
+		ProcessAudio(InNumInputChannels, InputAudio.GetData(), InNumOutputChannels, OutputAudio.GetData(), InputAudio.Num() / InNumInputChannels);
+	}
 
-	void FConvolutionReverb::ProcessAudio(int32 NumInputChannels, FAlignedFloatBuffer& InputAudio, int32 NumOutputChannels, FAlignedFloatBuffer& OutputAudio)
+	void FConvolutionReverb::ProcessAudio(int32 InNumInputChannels, const float* InputAudio, int32 InNumOutputChannels, float* OutputAudio, const int32 InNumFrames)
 	{
 		check(InputBlockBuffer.IsValid());
 		check(OutputBlockBuffer.IsValid());
 
-		OutputAudio.Reset();
+		if (InNumFrames < 1)
+		{
+			return;
+		}
 
-		const int32 NumInputSamples = InputAudio.Num();
-		const int32 NumFrames = NumInputSamples / NumInputChannels;
-		check((NumInputSamples % NumFrames) == 0);
-		const int32 NumOutputSamples = NumFrames * NumOutputChannels;
+		const int32 NumInputSamples = InNumFrames * InNumInputChannels;
+		const int32 NumOutputSamples = InNumFrames * InNumOutputChannels;
 
-		if (NumFrames != ExpectedNumFramesPerCallback)
+		if (InNumFrames != ExpectedNumFramesPerCallback)
 		{
-			ExpectedNumFramesPerCallback = NumFrames;
+			ExpectedNumFramesPerCallback = InNumFrames;
 			// If the number of frames processed per a call changes, then
 			// block buffers need to be updated. 
 			ResizeBlockBuffers();
 		}
 
-		if (NumFrames < 1)
-		{
-			return;
-		}
-
 		if (NumOutputSamples > 0)
 		{
-			OutputAudio.AddZeroed(NumOutputSamples);
+			FMemory::Memzero(OutputAudio, NumOutputSamples * sizeof(float));
 		}
 
 		// If we don't have a valid algorithm. We cannot do any processing. 
@@ -494,19 +497,19 @@ namespace Audio
 		}
 
 		// Our convolution algorithm should have equal number of input and output channels.
-		if (NumInputChannels != GetNumInputChannels())
+		if (InNumInputChannels != GetNumInputChannels())
 		{
-			UE_LOG(LogSynthesis, Error, TEXT("Convolution num input channel mismatch. Expected %d, got %d"), GetNumInputChannels(), NumInputChannels);
+			UE_LOG(LogSynthesis, Error, TEXT("Convolution num input channel mismatch. Expected %d, got %d"), GetNumInputChannels(), InNumInputChannels);
 			return;
 		}
-		else if (NumOutputChannels != GetNumOutputChannels())
+		else if (InNumOutputChannels != GetNumOutputChannels())
 		{
-			UE_LOG(LogSynthesis, Error, TEXT("Convolution num output channel mismatch. Expected %d, got %d"), GetNumOutputChannels(), NumOutputChannels);
+			UE_LOG(LogSynthesis, Error, TEXT("Convolution num output channel mismatch. Expected %d, got %d"), GetNumOutputChannels(), InNumOutputChannels);
 			return;
 		}
 
 		// Process all available audio that completes a block.
-		InputBlockBuffer->AddSamples(InputAudio.GetData(), NumInputSamples);
+		InputBlockBuffer->AddSamples(InputAudio, NumInputSamples);
 
 		while (InputBlockBuffer->GetNumAvailable() >= NumInputSamplesPerBlock)
 		{
@@ -514,7 +517,7 @@ namespace Audio
 
 			if (nullptr != InputBlock)
 			{
-				ProcessAudioBlock(InputBlock, NumInputChannels, InterleavedOutputBlock, NumOutputChannels);
+				ProcessAudioBlock(InputBlock, InNumInputChannels, InterleavedOutputBlock, InNumOutputChannels);
 				OutputBlockBuffer->AddSamples(InterleavedOutputBlock.GetData(), InterleavedOutputBlock.Num());
 			}
 
@@ -530,7 +533,7 @@ namespace Audio
 			const float* OutputPtr = OutputBlockBuffer->InspectSamples(NumToPop);
 			if (nullptr != OutputPtr)
 			{
-				FMemory::Memcpy(&OutputAudio.GetData()[OutputCopyOffset], OutputPtr, NumToPop * sizeof(float));
+				FMemory::Memcpy(&OutputAudio[OutputCopyOffset], OutputPtr, NumToPop * sizeof(float));
 			}
 			OutputBlockBuffer->RemoveSamples(NumToPop);
 		}
diff --git a/Engine/Plugins/Runtime/Synthesis/Source/Synthesis/Private/EffectConvolutionReverb.cpp b/Engine/Plugins/Runtime/Synthesis/Source/Synthesis/Private/EffectConvolutionReverb.cpp
new file mode 100644
--- /dev/null
+++ b/Engine/Plugins/Runtime/Synthesis/Source/Synthesis/Private/EffectConvolutionReverb.cpp
@@ -0,0 +1,361 @@
+#include "EffectConvolutionReverb.h"
+
+UAudioImpulseResponse::UAudioImpulseResponse()
+	: NumChannels(0)
+	, SampleRate(0)
+	, NormalizationVolumeDb(-24.f)
+{
+}
+
+#if WITH_EDITORONLY_DATA
+void UAudioImpulseResponse::PostEditChangeProperty(FPropertyChangedEvent& PropertyChangedEvent)
+{
+	Super::PostEditChangeProperty(PropertyChangedEvent);
+
+	OnObjectPropertyChanged.Broadcast(PropertyChangedEvent);
+}
+#endif
+
+namespace AudioConvReverbIntrinsics
+{
+	// Task for creating convolution algorithm object.
+	class SYNTHESIS_API FCreateConvolutionReverbTask : public FNonAbandonableTask
+	{
+		// This task can delete itself
+		friend class FAutoDeleteAsyncTask<FCreateConvolutionReverbTask>;
+
+	public:
+		FCreateConvolutionReverbTask(
+			TWeakPtr<Audio::FEffectConvolutionReverb, ESPMode::ThreadSafe> InEffectPtr,
+			Audio::FConvolutionReverbInitData&& InInitData,
+			Audio::FConvolutionReverbSettings& InSettings,
+			const FVersionData& InVersionData)
+			: EffectPtr(InEffectPtr)
+			, InitData(MoveTemp(InInitData))
+			, Settings(InSettings)
+			, VersionData(InVersionData)
+		{
+		}
+
+		void DoWork()
+		{
+			using namespace Audio;
+
+			// Build the convolution reverb object. 
+			TUniquePtr<FConvolutionReverb> ConvReverb = FConvolutionReverb::CreateConvolutionReverb(InitData, Settings);
+
+			TSharedPtr<FEffectConvolutionReverb, ESPMode::ThreadSafe> EffectSharedPtr = EffectPtr.Pin();
+			if (EffectSharedPtr.IsValid())
+			{
+				// Check that the effect still exists.
+				EffectSharedPtr->EnqueueNewReverb(MoveTemp(ConvReverb), VersionData);
+			}
+		}
+
+		FORCEINLINE TStatId GetStatId() const { RETURN_QUICK_DECLARE_CYCLE_STAT(CreateConvolutionReverbTask, STATGROUP_ThreadPoolAsyncTasks); }
+
+	private:
+		TWeakPtr<FEffectConvolutionReverb, ESPMode::ThreadSafe> EffectPtr;
+
+		FConvolutionReverbInitData InitData;
+		FConvolutionReverbSettings Settings;
+		FVersionData VersionData;
+	};
+
+	typedef FAutoDeleteAsyncTask<FCreateConvolutionReverbTask> FCreateReverbTask;
+}
+
+namespace Audio
+{
+	FConvolutionReverbInitData FEffectConvolutionReverb::CreateConvolutionReverbInitData()
+	{
+		FScopeLock ConvReverbInitDataLock(&ConvReverbInitDataCriticalSection);
+
+		const int32 NumInputChannelsLocal = NumInputChannels.Load();
+		const int32 NumOutputChannelsLocal = NumOutputChannels.Load();
+
+		ConvReverbInitData.InputAudioFormat.NumChannels = NumInputChannelsLocal;
+		ConvReverbInitData.OutputAudioFormat.NumChannels = NumOutputChannelsLocal;
+
+		// Mono to true stereo is a special case where we have the true stereo flag
+		// checked but only have 2 IRs
+		bool bIsMonoToTrueStereo = ConvReverbInitData.bIsImpulseTrueStereo && (ConvReverbInitData.AlgorithmSettings.NumImpulseResponses == 2);
+
+		// Determine correct input channel counts 
+		if (ConvReverbInitData.bMixInputChannelFormatToImpulseResponseFormat)
+		{
+			if (bIsMonoToTrueStereo)
+			{
+				// Force mono input.
+				ConvReverbInitData.AlgorithmSettings.NumInputChannels = 1;
+			}
+			else if (ConvReverbInitData.bIsImpulseTrueStereo)
+			{
+				// If performing true stereo, force input to be stereo
+				ConvReverbInitData.AlgorithmSettings.NumInputChannels = 2;
+			}
+			else
+			{
+				ConvReverbInitData.AlgorithmSettings.NumInputChannels = ConvReverbInitData.AlgorithmSettings.NumImpulseResponses;
+			}
+		}
+		else
+		{
+			ConvReverbInitData.AlgorithmSettings.NumInputChannels = NumInputChannelsLocal;
+		}
+
+
+		// Determine correct output channel count
+		if (bIsMonoToTrueStereo)
+		{
+			ConvReverbInitData.AlgorithmSettings.NumOutputChannels = 2;
+		}
+		else if (ConvReverbInitData.bIsImpulseTrueStereo)
+		{
+			// If IRs are true stereo, divide output channel count in half.
+			ConvReverbInitData.AlgorithmSettings.NumOutputChannels = ConvReverbInitData.AlgorithmSettings.NumImpulseResponses / 2;
+		}
+		else
+		{
+			ConvReverbInitData.AlgorithmSettings.NumOutputChannels = ConvReverbInitData.AlgorithmSettings.NumImpulseResponses;
+		}
+
+		// Determine longest IR
+		ConvReverbInitData.AlgorithmSettings.MaxNumImpulseResponseSamples = 0;
+
+		if (ConvReverbInitData.AlgorithmSettings.NumImpulseResponses > 0)
+		{
+			// Determine sample rate ratio in order to calculate the final IR num samples. 
+			float SampleRateRatio = 1.f;
+			if (ConvReverbInitData.ImpulseSampleRate > 0.f)
+			{
+				SampleRateRatio = ConvReverbInitData.TargetSampleRate / ConvReverbInitData.ImpulseSampleRate;
+			}
+
+			ConvReverbInitData.AlgorithmSettings.MaxNumImpulseResponseSamples = FMath::CeilToInt(SampleRateRatio * ConvReverbInitData.Samples.Num() / ConvReverbInitData.AlgorithmSettings.NumImpulseResponses) + 256;
+		}
+
+		// Setup gain matrix.
+		ConvReverbInitData.GainMatrix.Reset();
+
+		if (NumInputChannelsLocal == 1)
+		{
+			ConvReverbInitData.GainMatrix.Emplace(0, 0, 0, 1.f);
+		}
+		else if (bIsMonoToTrueStereo)
+		{
+			ConvReverbInitData.GainMatrix.Emplace(0, 0, 0, 1.f);
+			ConvReverbInitData.GainMatrix.Emplace(0, 1, 1, 1.f);
+		}
+		else if (ConvReverbInitData.bIsImpulseTrueStereo)
+		{
+			// True stereo treats first group of IRs 
+			int32 NumPairs = ConvReverbInitData.AlgorithmSettings.NumImpulseResponses / 2;
+
+			for (int32 PairIndex = 0; PairIndex < NumPairs; PairIndex++)
+			{
+				int32 LeftImpulseIndex = PairIndex;
+				int32 RightImpulseIndex = NumPairs + PairIndex;
+
+				ConvReverbInitData.GainMatrix.Emplace(0, LeftImpulseIndex, PairIndex, 1.f);
+				ConvReverbInitData.GainMatrix.Emplace(1, RightImpulseIndex, PairIndex, 1.f);
+			}
+		}
+		else
+		{
+			// Set up a 1-to-1 mapping. 
+			int32 MinChannelCount = FMath::Min3(ConvReverbInitData.AlgorithmSettings.NumInputChannels, ConvReverbInitData.AlgorithmSettings.NumOutputChannels, ConvReverbInitData.AlgorithmSettings.NumImpulseResponses);
+
+			for (int32 i = 0; i < MinChannelCount; i++)
+			{
+				ConvReverbInitData.GainMatrix.Emplace(i, i, i, 1.f);
+			}
+		}
+
+		return ConvReverbInitData;
+	}
+
+	void FEffectConvolutionReverb::BuildReverb()
+	{
+		using namespace AudioConvReverbIntrinsics;
+		FVersionData Version = UpdateVersion();
+
+		FConvolutionReverbSettings Settings;
+		Params.CopyParams(Settings);
+
+		FConvolutionReverbInitData InitData = CreateConvolutionReverbInitData();
+
+		(new FCreateReverbTask(AsWeak(), MoveTemp(InitData), Settings, VersionData))->StartBackgroundTask();
+	}
+
+	void FEffectConvolutionReverb::Init()
+	{
+		FConvolutionReverbSettings Settings;
+		Params.CopyParams(Settings);
+
+		// Create the convolution algorithm init data
+		FConvolutionReverbInitData ConvolutionInitData = CreateConvolutionReverbInitData();
+
+		SetReverb(FConvolutionReverb::CreateConvolutionReverb(ConvolutionInitData, Settings));
+	}
+
+	void FEffectConvolutionReverb::SetSettings(const FConvolutionReverbSettings& InSettings)
+	{
+		Params.SetParams(InSettings);
+
+		if (Reverb.IsValid())
+		{
+			Reverb->SetSettings(InSettings);
+		}
+	}
+
+	void FEffectConvolutionReverb::UpdateParameters()
+	{
+		Audio::FConvolutionReverbSettings NewSettings;
+		if (Reverb.IsValid() && Params.GetParams(&NewSettings))
+		{
+			Reverb->SetSettings(NewSettings);
+		}
+	}
+
+	void FEffectConvolutionReverb::SetSampleRate(const float InSampleRate)
+	{
+		FScopeLock Lock(&ConvReverbInitDataCriticalSection);
+
+		SampleRate = InSampleRate;
+		ConvReverbInitData.TargetSampleRate = SampleRate;
+	}
+
+	void FEffectConvolutionReverb::SetInitData(const FConvolutionReverbInitData& InData)
+	{
+		FScopeLock Lock(&ConvReverbInitDataCriticalSection);
+
+		ConvReverbInitData = InData;
+		ConvReverbInitData.TargetSampleRate = SampleRate;
+	}
+
+	void FEffectConvolutionReverb::SetReverb(TUniquePtr<Audio::FConvolutionReverb> InReverb)
+	{
+		Reverb = MoveTemp(InReverb);
+	}
+
+	void FEffectConvolutionReverb::SetBypass(const bool InBypass)
+	{
+		bBypass = InBypass;
+	}
+
+	void FEffectConvolutionReverb::UpdateChannelCount(const int32 InNumInputChannels, const int32 InNumOutputChannels)
+	{
+		int32 ExpectedNumInputChannels = 0;
+		int32 ExpectedNumOutputChannels = 0;
+
+		if (Reverb.IsValid())
+		{
+			ExpectedNumInputChannels = Reverb->GetNumInputChannels();
+			ExpectedNumOutputChannels = Reverb->GetNumOutputChannels();
+		}
+
+		// Check if there is a channel mismatch between the convolution reverb and the input/output data
+		const bool bNumChannelsMismatch = IsChannelCountUpToDate(InNumInputChannels, InNumOutputChannels);
+		const bool bRequestedChannelCountIsValid = InNumInputChannels > 0 && InNumOutputChannels > 0;
+
+		if (bNumChannelsMismatch && bRequestedChannelCountIsValid)
+		{
+			int32 NumInputChannelsTemp = NumInputChannels.Load();
+			int32 NumOutputChannelsTemp = NumOutputChannels.Load();
+
+			// check the NumInputChannelsTemp/NumOutputChannelsTemp to check if we've
+			// already tried and failed to create an algorithm for this channel configuration.
+			bool bShouldCreateNewAlgo = (NumInputChannelsTemp != InNumInputChannels) || (NumOutputChannelsTemp != InNumOutputChannels);
+
+			if (bShouldCreateNewAlgo)
+			{
+				UE_LOG(LogSynthesis, Log, TEXT("Creating new convolution algorithm due to channel count update. Num Inputs %d -> %d. Num Outputs %d -> %d")
+					, ExpectedNumInputChannels, InNumInputChannels, ExpectedNumOutputChannels, InNumOutputChannels);
+
+				// These should only be updated here and in the constructor since we use them
+				// to see if we've tried this channel configuration already.
+				NumInputChannels.Store(InNumInputChannels);
+				NumOutputChannels.Store(InNumOutputChannels);
+
+				// Create the convolution reverb in an async task.
+				BuildReverb();
+			}
+		}
+	}
+
+	void FEffectConvolutionReverb::EnqueueNewReverb(TUniquePtr<FConvolutionReverb> InReverb, const AudioConvReverbIntrinsics::FVersionData& InVersionData)
+	{
+		// This queued reverb will get hotswapped the next time ProcessAudio comes around on the render thread
+		FScopeLock Lock(&VersionDataCriticalSection);
+
+		QueuedVersionData = InVersionData;
+		QueuedReverb = MoveTemp(InReverb);
+	}
+
+	void FEffectConvolutionReverb::DequeueNewReverb()
+	{
+		FScopeLock Lock(&VersionDataCriticalSection);
+		
+		if (QueuedReverb.IsValid() && VersionData == QueuedVersionData)
+		{
+			Reverb = MoveTemp(QueuedReverb);
+		}
+	}
+
+	bool FEffectConvolutionReverb::IsVersionCurrent(const AudioConvReverbIntrinsics::FVersionData& InVersionData) const
+	{
+		FScopeLock VersionDataLock(&VersionDataCriticalSection);
+
+		return InVersionData == VersionData;
+	}
+
+	AudioConvReverbIntrinsics::FVersionData FEffectConvolutionReverb::UpdateVersion()
+	{
+		AudioConvReverbIntrinsics::FVersionData VersionDataCopy;
+		FScopeLock VersionDataLock(&VersionDataCriticalSection);
+
+		VersionData.ConvolutionID++;
+		VersionDataCopy = VersionData;
+
+		return VersionDataCopy;
+	}
+
+	bool FEffectConvolutionReverb::IsChannelCountUpToDate(const int32 InNumInputChannels, const int32 InNumOutputChannels) const
+	{
+		if (Reverb.IsValid() == false)
+		{
+			return false;
+		}
+
+		const int32 ExpectedNumInputChannels = Reverb->GetNumInputChannels();
+		const int32 ExpectedNumOutputChannels = Reverb->GetNumOutputChannels();
+
+		return ExpectedNumInputChannels == InNumInputChannels && ExpectedNumOutputChannels == InNumOutputChannels;
+	}
+
+	void FEffectConvolutionReverb::ProcessAudio(int32 InNumInputChannels, const float* InputAudio, int32 InNumOutputChannels, float* OutputAudio, const int32 InNumFrames)
+	{
+		DequeueNewReverb();
+
+		const bool bShouldProcessConvReverb = !bBypass && Reverb.IsValid() && IsChannelCountUpToDate(InNumInputChannels, InNumOutputChannels);
+
+		if (bShouldProcessConvReverb)
+		{
+			UpdateParameters();
+			Reverb->ProcessAudio(InNumInputChannels, InputAudio, InNumOutputChannels, OutputAudio, InNumFrames);
+		}
+		else if (bBypass)
+		{
+			FMemory::Memcpy(OutputAudio, InputAudio, InNumFrames * InNumOutputChannels * sizeof(float));
+		}
+		else
+		{
+			// Zero output data. Do *not* trigger rebuild here in case one is already in flight or simply because one cannot be built.
+			if (InNumFrames > 0)
+			{
+				FMemory::Memzero(OutputAudio, InNumFrames * InNumOutputChannels * sizeof(float));
+			}
+		}
+	}
+}
\ No newline at end of file
diff --git a/Engine/Plugins/Runtime/Synthesis/Source/Synthesis/Private/SourceEffectConvolutionReverb.cpp b/Engine/Plugins/Runtime/Synthesis/Source/Synthesis/Private/SourceEffectConvolutionReverb.cpp
new file mode 100644
--- /dev/null
+++ b/Engine/Plugins/Runtime/Synthesis/Source/Synthesis/Private/SourceEffectConvolutionReverb.cpp
@@ -0,0 +1,393 @@
+// Copyright Epic Games, Inc. All Rights Reserved.
+
+#include "SourceEffects/SourceEffectConvolutionReverb.h"
+
+#include "CoreMinimal.h"
+#include "Async/Async.h"
+#include "SynthesisModule.h"
+
+FSourceEffectConvolutionReverbSettings::FSourceEffectConvolutionReverbSettings()
+	: NormalizationVolumeDb(-24.f)
+	, WetVolumeDb(0.f)
+	, DryVolumeDb(-96.f)
+	, bBypass(false)
+{
+}
+
+FSourceEffectConvolutionReverb::FSourceEffectConvolutionReverb(const USourceEffectConvolutionReverbPreset* InPreset)
+	: Reverb(MakeShared<Audio::FEffectConvolutionReverb, ESPMode::ThreadSafe>())
+{
+	UpdateConvolutionReverb(InPreset);
+
+	if (nullptr != InPreset)
+	{
+		SetConvolutionReverbParameters(InPreset->GetSettings());
+	}
+}
+
+FSourceEffectConvolutionReverb::~FSourceEffectConvolutionReverb()
+{
+}
+
+void FSourceEffectConvolutionReverb::Init(const FSoundEffectSourceInitData& InitData)
+{
+	using namespace Audio;
+
+	NumChannels = InitData.NumSourceChannels;
+
+	Reverb->SetSampleRate(InitData.SampleRate);
+	Reverb->UpdateVersion();
+	Reverb->Init();
+	Reverb->UpdateChannelCount(NumChannels, NumChannels);
+}
+
+
+void FSourceEffectConvolutionReverb::OnPresetChanged()
+{
+	USourceEffectConvolutionReverbPreset* ConvolutionPreset = CastChecked<USourceEffectConvolutionReverbPreset>(Preset);
+	FSourceEffectConvolutionReverbSettings Settings = ConvolutionPreset->GetSettings();
+
+	// Copy settings from FSourceEffectConvolutionReverbSettings needed for FConvolutionReverbSettings 
+	// FConvolutionReverbSettings represents runtime settings which do not need the 
+	// FConvlutionReverb object to be rebuilt. Some settings in FSourceEffectConvolutionReverbSettings
+	// force a rebuild of FConvolutionReverb-> Those are handled in USourceEffectConvolutinReverbPreset::PostEditChangeProperty
+	SetConvolutionReverbParameters(Settings);
+}
+
+void FSourceEffectConvolutionReverb::SetConvolutionReverbParameters(const FSourceEffectConvolutionReverbSettings& InSettings)
+{
+	Audio::FConvolutionReverbSettings ReverbSettings;
+
+	ReverbSettings.NormalizationVolume = Audio::ConvertToLinear(InSettings.NormalizationVolumeDb);
+
+	WetVolume = Audio::ConvertToLinear(InSettings.WetVolumeDb);
+	DryVolume = Audio::ConvertToLinear(InSettings.DryVolumeDb);
+
+	Reverb->SetSettings(ReverbSettings);
+	Reverb->SetBypass(InSettings.bBypass);
+}
+
+void FSourceEffectConvolutionReverb::ProcessAudio(const FSoundEffectSourceInputData& InData, float* OutAudioBufferData)
+{
+	using namespace Audio;
+
+	check(nullptr != InData.InputSourceEffectBufferPtr);
+	check(nullptr != OutAudioBufferData);
+	check(NumChannels != 0);
+
+	const int32 NumFrames = InData.NumSamples / NumChannels;
+
+	Reverb->UpdateChannelCount(NumChannels, NumChannels);
+
+	Reverb->ProcessAudio(NumChannels, InData.InputSourceEffectBufferPtr, NumChannels, OutAudioBufferData, NumFrames);
+
+	// Process Wet/Dry mix
+	if (FMath::IsNearlyEqual(WetVolume, 1.f) == false)
+	{
+		Audio::MultiplyBufferByConstantInPlace(OutAudioBufferData, InData.NumSamples, WetVolume);
+	}
+	if (DryVolume > KINDA_SMALL_NUMBER)
+	{
+		Audio::MixInBufferFast(InData.InputSourceEffectBufferPtr, OutAudioBufferData, InData.NumSamples, DryVolume);
+	}
+}
+
+AudioConvReverbIntrinsics::FVersionData FSourceEffectConvolutionReverb::UpdateConvolutionReverb(const USourceEffectConvolutionReverbPreset* InPreset)
+{
+	using namespace Audio;
+
+	// Copy data from preset into internal FConvolutionReverbInitData
+	// Reset data
+	FConvolutionReverbInitData InitData = FConvolutionReverbInitData();
+
+	if (nullptr != InPreset)
+	{
+		if (InPreset->ImpulseResponse)
+		{
+			UAudioImpulseResponse* IR = InPreset->ImpulseResponse;
+
+			InitData.Samples = IR->ImpulseResponse;
+			InitData.AlgorithmSettings.NumImpulseResponses = IR->NumChannels;
+			InitData.ImpulseSampleRate = IR->SampleRate;
+			InitData.bIsImpulseTrueStereo = IR->bTrueStereo && (IR->NumChannels % 2 == 0);
+			InitData.NormalizationVolume = ConvertToLinear(IR->NormalizationVolumeDb);
+		}
+
+		switch (InPreset->BlockSize)
+		{
+		case ESubmixEffectConvolutionReverbBlockSize::BlockSize256:
+			InitData.AlgorithmSettings.BlockNumSamples = 256;
+			break;
+
+		case ESubmixEffectConvolutionReverbBlockSize::BlockSize512:
+			InitData.AlgorithmSettings.BlockNumSamples = 512;
+			break;
+
+		case ESubmixEffectConvolutionReverbBlockSize::BlockSize1024:
+		default:
+			InitData.AlgorithmSettings.BlockNumSamples = 1024;
+			break;
+		}
+
+		InitData.AlgorithmSettings.bEnableHardwareAcceleration = InPreset->bEnableHardwareAcceleration;
+	}
+
+	Reverb->SetInitData(InitData);
+	return Reverb->UpdateVersion();
+}
+
+void FSourceEffectConvolutionReverb::RebuildConvolutionReverb()
+{
+	Reverb->BuildReverb();
+}
+
+/********************************************************************/
+/***************** USourceEffectConvolutionReverbPreset *************/
+/********************************************************************/
+
+USourceEffectConvolutionReverbPreset::USourceEffectConvolutionReverbPreset(const FObjectInitializer& ObjectInitializer)
+:	Super(ObjectInitializer)
+,	ImpulseResponse(nullptr)
+,	BlockSize(ESubmixEffectConvolutionReverbBlockSize::BlockSize1024)
+,	bEnableHardwareAcceleration(true)
+{}
+
+bool USourceEffectConvolutionReverbPreset::CanFilter() const 
+{
+	return false; 
+}
+
+bool USourceEffectConvolutionReverbPreset::HasAssetActions() const 
+{ 
+	return true; 
+}
+
+FText USourceEffectConvolutionReverbPreset::GetAssetActionName() const 
+{ 
+	return FText::FromString(TEXT("SourceEffectConvolutionReverb")); 
+}
+
+UClass* USourceEffectConvolutionReverbPreset::GetSupportedClass() const 
+{ 
+	return USourceEffectConvolutionReverbPreset::StaticClass(); 
+}
+
+FSoundEffectBase* USourceEffectConvolutionReverbPreset::CreateNewEffect() const 
+{ 
+	// Pass a pointer to self into this constructor.  This enables the effect to have
+	// a copy of the impulse response data. The effect can then prepare the impulse
+	// response during it's "Init" routine.
+	return new FSourceEffectConvolutionReverb(this); 
+}
+
+USoundEffectPreset* USourceEffectConvolutionReverbPreset::CreateNewPreset(UObject* InParent, FName Name, EObjectFlags Flags) const
+{ 
+	USoundEffectPreset* NewPreset = NewObject<USourceEffectConvolutionReverbPreset>(InParent, GetSupportedClass(), Name, Flags); 
+	NewPreset->Init(); 
+	return NewPreset; 
+}
+
+void USourceEffectConvolutionReverbPreset::Init() 
+{ 
+	FScopeLock ScopeLock(&SettingsCritSect);
+	SettingsCopy = Settings;
+}
+
+void USourceEffectConvolutionReverbPreset::UpdateSettings()
+{
+	{
+		// Copy settings to audio-render-thread version
+		FScopeLock ScopeLock(&SettingsCritSect);
+		SettingsCopy = Settings; 
+	}
+	// This sets the 'bChanged' to true on related effect instances which will 
+	// trigger a OnPresetChanged call on the audio render thread.
+	Update(); 
+} 
+
+FSourceEffectConvolutionReverbSettings USourceEffectConvolutionReverbPreset::GetSettings() const
+{ 
+	FScopeLock ScopeLock(&SettingsCritSect);
+	return SettingsCopy; 
+} 
+
+void USourceEffectConvolutionReverbPreset::SetSettings(const FSourceEffectConvolutionReverbSettings& InSettings)
+{
+	Settings = InSettings;
+
+	if (nullptr != ImpulseResponse)
+	{
+		SetImpulseResponseSettings(ImpulseResponse);
+	}
+
+	UpdateSettings();
+}
+
+void USourceEffectConvolutionReverbPreset::SetImpulseResponse(UAudioImpulseResponse* InImpulseResponse)
+{
+	ImpulseResponse = InImpulseResponse;
+	SetImpulseResponseSettings(InImpulseResponse);
+	RebuildConvolutionReverb();
+}
+
+void USourceEffectConvolutionReverbPreset::RebuildConvolutionReverb()
+{
+	using namespace AudioConvReverbIntrinsics;
+	using FEffectWeakPtr = TWeakPtr<FSoundEffectBase, ESPMode::ThreadSafe>;
+	using FEffectSharedPtr = TSharedPtr<FSoundEffectBase, ESPMode::ThreadSafe>;
+	using FSourceEffectConvSharedPtr = TSharedPtr<FSourceEffectConvolutionReverb, ESPMode::ThreadSafe>;
+
+	// Iterator over effects and update the convolution Reverb->
+	for (FEffectWeakPtr& InstanceWeakPtr : Instances)
+	{
+		FEffectSharedPtr InstanceSharedPtr = InstanceWeakPtr.Pin();
+		if (InstanceSharedPtr.IsValid())
+		{
+			FSourceEffectConvSharedPtr ConvEffectSharedPtr = StaticCastSharedPtr<FSourceEffectConvolutionReverb>(InstanceSharedPtr);
+
+			FVersionData VersionData = ConvEffectSharedPtr->UpdateConvolutionReverb(this);
+			ConvEffectSharedPtr->RebuildConvolutionReverb();
+		}
+	}
+}
+
+#if WITH_EDITORONLY_DATA
+void USourceEffectConvolutionReverbPreset::BindToImpulseResponseObjectChange()
+{
+	if (nullptr != ImpulseResponse)
+	{
+		if (!DelegateHandles.Contains(ImpulseResponse))
+		{
+			FDelegateHandle Handle = ImpulseResponse->OnObjectPropertyChanged.AddUObject(this, &USourceEffectConvolutionReverbPreset::PostEditChangeImpulseProperty);
+			if (Handle.IsValid())
+			{
+				DelegateHandles.Add(ImpulseResponse, Handle);
+			}
+		}
+	}
+}
+
+void USourceEffectConvolutionReverbPreset::PreEditChange(FProperty* PropertyAboutToChange)
+{
+	Super::PreEditChange(PropertyAboutToChange);
+
+	if (nullptr == PropertyAboutToChange)
+	{
+		return;
+	}
+
+	const FName ImpulseResponseFName = GET_MEMBER_NAME_CHECKED(USourceEffectConvolutionReverbPreset, ImpulseResponse);
+
+	if (PropertyAboutToChange->GetFName() == ImpulseResponseFName)
+	{
+		if (nullptr == ImpulseResponse)
+		{
+			return;
+		}
+
+		if (DelegateHandles.Contains(ImpulseResponse))
+		{
+			FDelegateHandle DelegateHandle = DelegateHandles[ImpulseResponse];
+
+			if (DelegateHandle.IsValid())
+			{
+				ImpulseResponse->OnObjectPropertyChanged.Remove(DelegateHandle);
+			}
+
+			DelegateHandles.Remove(ImpulseResponse);
+		}
+	}
+}
+
+
+
+void USourceEffectConvolutionReverbPreset::PostEditChangeProperty(FPropertyChangedEvent& PropertyChangedEvent)
+{
+	const FName ImpulseResponseFName = GET_MEMBER_NAME_CHECKED(USourceEffectConvolutionReverbPreset, ImpulseResponse);
+
+	// If any of these properties are updated, the FConvolutionReverb object needs
+	// to be rebuilt.
+	const TArray<FName> ConvolutionReverbProperties(
+		{
+			ImpulseResponseFName,
+			GET_MEMBER_NAME_CHECKED(USourceEffectConvolutionReverbPreset, bEnableHardwareAcceleration),
+			GET_MEMBER_NAME_CHECKED(USourceEffectConvolutionReverbPreset, BlockSize),
+		}
+	);
+
+	if (!HasAnyFlags(RF_ClassDefaultObject))
+	{
+		if (FProperty* PropertyThatChanged = PropertyChangedEvent.Property)
+		{
+			const FName& Name = PropertyThatChanged->GetFName();
+
+			// Need to update impulse reponse bindings and normalization before rebuilding convolution reverb
+			if (Name == ImpulseResponseFName)
+			{
+				BindToImpulseResponseObjectChange();
+				
+				if (nullptr != ImpulseResponse)
+				{
+					SetImpulseResponseSettings(ImpulseResponse);
+				}
+			}
+
+			if (ConvolutionReverbProperties.Contains(Name))
+			{
+				RebuildConvolutionReverb();
+			}
+		}
+	}
+
+	Super::PostEditChangeProperty(PropertyChangedEvent);
+}
+
+void USourceEffectConvolutionReverbPreset::PostEditChangeImpulseProperty(FPropertyChangedEvent& PropertyChangedEvent) 
+{
+	const FName NormalizationVolumeFName = GET_MEMBER_NAME_CHECKED(UAudioImpulseResponse, NormalizationVolumeDb);
+	const FName TrueStereoFName = GET_MEMBER_NAME_CHECKED(UAudioImpulseResponse, bTrueStereo);
+
+	if (!HasAnyFlags(RF_ClassDefaultObject))
+	{
+		if (FProperty* PropertyThatChanged = PropertyChangedEvent.Property)
+		{
+			const FName& Name = PropertyThatChanged->GetFName();
+
+			if (Name == NormalizationVolumeFName)
+			{
+				if (nullptr != ImpulseResponse)
+				{
+					Settings.NormalizationVolumeDb = ImpulseResponse->NormalizationVolumeDb;
+					UpdateSettings();
+				}
+			}
+			else if (Name == TrueStereoFName)
+			{
+				RebuildConvolutionReverb();
+			}
+		}
+	}
+}
+
+#endif
+
+void USourceEffectConvolutionReverbPreset::SetImpulseResponseSettings(UAudioImpulseResponse* InImpulseResponse)
+{
+	if (nullptr != InImpulseResponse)
+	{
+		// Set this value, but do not call UpdateSettings(). UpdateSettings() is handled elsewhere.
+		Settings.NormalizationVolumeDb = InImpulseResponse->NormalizationVolumeDb;
+	}
+}
+
+void USourceEffectConvolutionReverbPreset::PostLoad()
+{
+	Super::PostLoad();
+
+#if WITH_EDITORONLY_DATA
+	// Bind to trigger new convolution algorithms when UAudioImpulseResponse object changes.
+	BindToImpulseResponseObjectChange();
+#endif
+
+	SetImpulseResponseSettings(ImpulseResponse);
+}
diff --git a/Engine/Plugins/Runtime/Synthesis/Source/Synthesis/Private/SubmixEffectConvolutionReverb.cpp b/Engine/Plugins/Runtime/Synthesis/Source/Synthesis/Private/SubmixEffectConvolutionReverb.cpp
--- a/Engine/Plugins/Runtime/Synthesis/Source/Synthesis/Private/SubmixEffectConvolutionReverb.cpp
+++ b/Engine/Plugins/Runtime/Synthesis/Source/Synthesis/Private/SubmixEffectConvolutionReverb.cpp
@@ -4,116 +4,26 @@
 
 #include "CoreMinimal.h"
 #include "Async/Async.h"
-#include "ConvolutionReverb.h"
-#include "DSP/ConvolutionAlgorithm.h"
 #include "SynthesisModule.h"
 
-namespace AudioConvReverbIntrinsics
-{
-	using namespace Audio;
-
-	// names of various smart pointers
-	using FSubmixEffectConvSharedPtr = TSharedPtr<FSubmixEffectConvolutionReverb, ESPMode::ThreadSafe>;
-	using FSubmixEffectConvWeakPtr = TWeakPtr<FSubmixEffectConvolutionReverb, ESPMode::ThreadSafe>;
-
-	using FVersionData = FSubmixEffectConvolutionReverb::FVersionData;
-
-	// Task for creating convolution algorithm object.
-	class FCreateConvolutionReverbTask : public FNonAbandonableTask
-	{
-		// This task can delete itself
-		friend class FAutoDeleteAsyncTask<FCreateConvolutionReverbTask>;
-
-		public:
-			FCreateConvolutionReverbTask(
-					FSubmixEffectConvWeakPtr InEffectObject,
-					FConvolutionReverbInitData&& InInitData,
-					FConvolutionReverbSettings& InSettings,
-					const FVersionData& InVersionData)
-			:	EffectWeakPtr(InEffectObject)
-			,	InitData(MoveTemp(InInitData))
-			,	Settings(InSettings)
-			,	VersionData(InVersionData)
-			{}
-
-			void DoWork()
-			{
-				// Build the convolution reverb object. 
-				TUniquePtr<FConvolutionReverb> ConvReverb = FConvolutionReverb::CreateConvolutionReverb(InitData, Settings);
-
-				FSubmixEffectConvSharedPtr EffectSharedPtr = EffectWeakPtr.Pin();
-
-				// Check that the submix effect still exists.
-				if (EffectSharedPtr.IsValid())
-				{
-					// If the effect ptr is still valid, add to it's command queue to set the convolution
-					// algorithm object in the audio render thread.
-					TUniqueFunction<void()> Command = [Algo = MoveTemp(ConvReverb), EffectSharedPtr, VersionData = VersionData] () mutable
-					{
-						EffectSharedPtr->SetConvolutionReverbIfCurrent(MoveTemp(Algo), VersionData);
-					};
-					EffectSharedPtr->EffectCommand(MoveTemp(Command));
-				}
-			}
-
-			FORCEINLINE TStatId GetStatId() const { RETURN_QUICK_DECLARE_CYCLE_STAT(CreateConvolutionReverbTask,     STATGROUP_ThreadPoolAsyncTasks); }
-
-		private:
-			FSubmixEffectConvWeakPtr EffectWeakPtr;
-
-			FConvolutionReverbInitData InitData;
-			FConvolutionReverbSettings Settings;
-			FVersionData VersionData;
-	};
-}
-
-UAudioImpulseResponse::UAudioImpulseResponse()
-:	NumChannels(0)
-,	SampleRate(0)
-,	NormalizationVolumeDb(-24.f)
-{
-}
-
-#if WITH_EDITORONLY_DATA
-void UAudioImpulseResponse::PostEditChangeProperty(FPropertyChangedEvent& PropertyChangedEvent) 
-{
-	Super::PostEditChangeProperty(PropertyChangedEvent);
-
-	OnObjectPropertyChanged.Broadcast(PropertyChangedEvent);
-}
-#endif
-
 FSubmixEffectConvolutionReverbSettings::FSubmixEffectConvolutionReverbSettings()
-:	NormalizationVolumeDb(-24.f)
-,	bBypass(false)
-,	bMixInputChannelFormatToImpulseResponseFormat(true)
-,	bMixReverbOutputToOutputChannelFormat(true)
-,	SurroundRearChannelBleedDb(-60.f)
-,	bInvertRearChannelBleedPhase(false)
-,	bSurroundRearChannelFlip(false)
-,	SurroundRearChannelBleedAmount_DEPRECATED(0.0f)
-,	ImpulseResponse_DEPRECATED(nullptr)
-,	AllowHardwareAcceleration_DEPRECATED(true)
-{
-}
-
-FSubmixEffectConvolutionReverb::FVersionData::FVersionData()
-: 	ConvolutionID(0)
+	: NormalizationVolumeDb(-24.f)
+	, WetVolumeDb(0.f)
+	, DryVolumeDb(-96.f)
+	, bBypass(false)
+	, bMixInputChannelFormatToImpulseResponseFormat(true)
+	, bMixReverbOutputToOutputChannelFormat(true)
+	, SurroundRearChannelBleedDb(-60.f)
+	, bInvertRearChannelBleedPhase(false)
+	, bSurroundRearChannelFlip(false)
+	, SurroundRearChannelBleedAmount_DEPRECATED(0.0f)
+	, ImpulseResponse_DEPRECATED(nullptr)
+	, AllowHardwareAcceleration_DEPRECATED(true)
 {
 }
 
-bool FSubmixEffectConvolutionReverb::FVersionData::operator==(const FVersionData& Other) const
-{
-	// Check if versions are equal
-	return ConvolutionID == Other.ConvolutionID;
-}
-
 FSubmixEffectConvolutionReverb::FSubmixEffectConvolutionReverb(const USubmixEffectConvolutionReverbPreset* InPreset)
-: 	ConvolutionReverb(TUniquePtr<Audio::FConvolutionReverb>(nullptr))
-,	SampleRate(0.0f)
-,	NumInputChannels(2)
-,	NumOutputChannels(2)
-,	bBypass(false)
+	: Reverb(MakeShared<Audio::FEffectConvolutionReverb, ESPMode::ThreadSafe>())
 {
 	UpdateConvolutionReverb(InPreset);
 
@@ -131,24 +41,12 @@ void FSubmixEffectConvolutionReverb::Init(const FSoundEffectSubmixInitData& Init
 {
 	using namespace Audio;
 
-	SampleRate = InitData.SampleRate;
-
-	FVersionData UpdatedVersionData = UpdateVersion();
+	Reverb->SetSampleRate(InitData.SampleRate);
 
-	// Create the convolution algorithm init data
-	FConvolutionReverbInitData ConvolutionInitData = CreateConvolutionReverbInitData();
-
-	FConvolutionReverbSettings Settings;
-	Params.CopyParams(Settings);
-
-	// Create the convolution reverb
-	TUniquePtr<FConvolutionReverb> ConvReverb = FConvolutionReverb::CreateConvolutionReverb(ConvolutionInitData, Settings);
-
-	// Set the convolution algorithm, ensuring that it's the most up-to-date
-	SetConvolutionReverbIfCurrent(MoveTemp(ConvReverb), UpdatedVersionData);
+	Reverb->UpdateVersion();
+	Reverb->Init();
 }
 
-
 void FSubmixEffectConvolutionReverb::OnPresetChanged()
 {
 	USubmixEffectConvolutionReverbPreset* ConvolutionPreset = CastChecked<USubmixEffectConvolutionReverbPreset>(Preset);
@@ -157,7 +55,7 @@ void FSubmixEffectConvolutionReverb::OnPresetChanged()
 	// Copy settings from FSubmixEffectConvolutionReverbSettings needed for FConvolutionReverbSettings 
 	// FConvolutionReverbSettings represents runtime settings which do not need the 
 	// FConvlutionReverb object to be rebuilt. Some settings in FSubmixEffectConvolutionReverbSettings
-	// force a rebuild of FConvolutionReverb. Those are handled in USubmixEffectConvolutinReverbPreset::PostEditChangeProperty
+	// force a rebuild of FConvolutionReverb-> Those are handled in USubmixEffectConvolutinReverbPreset::PostEditChangeProperty
 	SetConvolutionReverbParameters(Settings);
 }
 
@@ -165,7 +63,6 @@ void FSubmixEffectConvolutionReverb::SetConvolutionReverbParameters(const FSubmi
 {
 	Audio::FConvolutionReverbSettings ReverbSettings;
 
-	float NewVolume = Audio::ConvertToLinear(InSettings.NormalizationVolumeDb);
 	float NewRearChannelBleed = Audio::ConvertToLinear(InSettings.SurroundRearChannelBleedDb);
 
 	if (InSettings.bInvertRearChannelBleedPhase)
@@ -173,128 +70,16 @@ void FSubmixEffectConvolutionReverb::SetConvolutionReverbParameters(const FSubmi
 		NewRearChannelBleed *= -1.f;
 	}
 
-	ReverbSettings.NormalizationVolume = NewVolume;
 	ReverbSettings.RearChannelBleed = NewRearChannelBleed;
-	ReverbSettings.bRearChannelFlip = InSettings.bSurroundRearChannelFlip;
-
-	Params.SetParams(ReverbSettings);
-
-	bBypass = InSettings.bBypass;
-}
-
-
-Audio::FConvolutionReverbInitData FSubmixEffectConvolutionReverb::CreateConvolutionReverbInitData() 
-{
-	FScopeLock ConvReverbInitDataLock(&ConvReverbInitDataCriticalSection);
-
-	int32 NumInputChannelsLocal = NumInputChannels.Load();
-	int32 NumOutputChannelsLocal = NumOutputChannels.Load();
-
-	ConvReverbInitData.InputAudioFormat.NumChannels = NumInputChannelsLocal;
-	ConvReverbInitData.OutputAudioFormat.NumChannels = NumOutputChannelsLocal;
-	ConvReverbInitData.TargetSampleRate = SampleRate;
-
-	// Mono to true stereo is a special case where we have the true stereo flag
-	// checked but only have 2 IRs
-	bool bIsMonoToTrueStereo = ConvReverbInitData.bIsImpulseTrueStereo && (ConvReverbInitData.AlgorithmSettings.NumImpulseResponses == 2);
-
-	// Determine correct input channel counts 
-	if (ConvReverbInitData.bMixInputChannelFormatToImpulseResponseFormat)
-	{
-		if (bIsMonoToTrueStereo)
-		{
-			// Force mono input.
-			ConvReverbInitData.AlgorithmSettings.NumInputChannels = 1;
-		}
-		else if (ConvReverbInitData.bIsImpulseTrueStereo)
-		{
-			// If performing true stereo, force input to be stereo
-			ConvReverbInitData.AlgorithmSettings.NumInputChannels = 2;
-		}
-		else
-		{
-			ConvReverbInitData.AlgorithmSettings.NumInputChannels = ConvReverbInitData.AlgorithmSettings.NumImpulseResponses;
-		}
-	}
-	else
-	{
-		ConvReverbInitData.AlgorithmSettings.NumInputChannels = NumInputChannelsLocal;
-	}
-
-
-	// Determine correct output channel count
-	if (bIsMonoToTrueStereo)
-	{
-		ConvReverbInitData.AlgorithmSettings.NumOutputChannels = 2;
-	}
-	else if (ConvReverbInitData.bIsImpulseTrueStereo)
-	{
-		// If IRs are true stereo, divide output channel count in half.
-		ConvReverbInitData.AlgorithmSettings.NumOutputChannels = ConvReverbInitData.AlgorithmSettings.NumImpulseResponses / 2;
-	}
-	else
-	{
-		ConvReverbInitData.AlgorithmSettings.NumOutputChannels = ConvReverbInitData.AlgorithmSettings.NumImpulseResponses;
-	}
-
-	// Determine longest IR
-	ConvReverbInitData.AlgorithmSettings.MaxNumImpulseResponseSamples = 0;
-
-	if (ConvReverbInitData.AlgorithmSettings.NumImpulseResponses > 0)
-	{
-		// Determine sample rate ratio in order to calculate the final IR num samples. 
-		float SampleRateRatio = 1.f;
-		if (ConvReverbInitData.ImpulseSampleRate > 0.f)
-		{
-			SampleRateRatio = ConvReverbInitData.TargetSampleRate / ConvReverbInitData.ImpulseSampleRate;
-		}
-
-		ConvReverbInitData.AlgorithmSettings.MaxNumImpulseResponseSamples = FMath::CeilToInt(SampleRateRatio * ConvReverbInitData.Samples.Num() / ConvReverbInitData.AlgorithmSettings.NumImpulseResponses) + 256;
-	}
-
-	// Setup gain matrix.
-	ConvReverbInitData.GainMatrix.Reset();
-
-	if (bIsMonoToTrueStereo)
-	{
-		ConvReverbInitData.GainMatrix.Emplace(0, 0, 0, 1.f);
-		ConvReverbInitData.GainMatrix.Emplace(0, 1, 1, 1.f);
-	}
-	else if (ConvReverbInitData.bIsImpulseTrueStereo)
-	{
-		// True stereo treats first group of IRs 
-		int32 NumPairs = ConvReverbInitData.AlgorithmSettings.NumImpulseResponses / 2;
 
-		for (int32 PairIndex = 0; PairIndex < NumPairs; PairIndex++)
-		{
-			int32 LeftImpulseIndex = PairIndex;
-			int32 RightImpulseIndex = NumPairs + PairIndex;
-
-			ConvReverbInitData.GainMatrix.Emplace(0, LeftImpulseIndex,  PairIndex, 1.f);
-			ConvReverbInitData.GainMatrix.Emplace(1, RightImpulseIndex, PairIndex, 1.f);
-		}
-	}
-	else
-	{
-		// Set up a 1-to-1 mapping. 
-		int32 MinChannelCount = FMath::Min3(ConvReverbInitData.AlgorithmSettings.NumInputChannels, ConvReverbInitData.AlgorithmSettings.NumOutputChannels, ConvReverbInitData.AlgorithmSettings.NumImpulseResponses);
-
-		for (int32 i = 0; i < MinChannelCount; i++)
-		{
-			ConvReverbInitData.GainMatrix.Emplace(i, i, i, 1.f);
-		}
-	}
-
-	return ConvReverbInitData;
-}
-
-Audio::FConvolutionReverbSettings FSubmixEffectConvolutionReverb::GetParameters() const
-{
-	Audio::FConvolutionReverbSettings Settings;
+	ReverbSettings.NormalizationVolume = Audio::ConvertToLinear(InSettings.NormalizationVolumeDb);
+	ReverbSettings.bRearChannelFlip = InSettings.bSurroundRearChannelFlip;
 
-	Params.CopyParams(Settings);
+	WetVolume = Audio::ConvertToLinear(InSettings.WetVolumeDb);
+	DryVolume = Audio::ConvertToLinear(InSettings.DryVolumeDb);
 
-	return Settings;
+	Reverb->SetSettings(ReverbSettings);
+	Reverb->SetBypass(InSettings.bBypass);
 }
 
 void FSubmixEffectConvolutionReverb::OnProcessAudio(const FSoundEffectSubmixInputData& InData, FSoundEffectSubmixOutputData& OutData)
@@ -303,199 +88,72 @@ void FSubmixEffectConvolutionReverb::OnProcessAudio(const FSoundEffectSubmixInpu
 
 	check(nullptr != InData.AudioBuffer);
 	check(nullptr != OutData.AudioBuffer);
+	check(InData.NumChannels != 0);
 
-	UpdateParameters();
-
-	int32 ExpectedNumInputChannels = 0;
-	int32 ExpectedNumOutputChannels = 0;
-
-	if (ConvolutionReverb.IsValid())
-	{
-		ExpectedNumInputChannels = ConvolutionReverb->GetNumInputChannels();
-		ExpectedNumOutputChannels = ConvolutionReverb->GetNumOutputChannels();
-	}
-
-	// Check if there is a channel mismatch between the convolution reverb and the input/output data
-	bool bNumChannelsMismatch = false;
-	bNumChannelsMismatch |= (ExpectedNumInputChannels != InData.NumChannels);
-	bNumChannelsMismatch |= (ExpectedNumOutputChannels != OutData.NumChannels);
-
-	if (bNumChannelsMismatch)
-	{
-		int32 NumInputChannelsTemp = NumInputChannels.Load();
-		int32 NumOutputChannelsTemp = NumOutputChannels.Load();
+	Reverb->UpdateChannelCount(InData.NumChannels, OutData.NumChannels);
 
-		// We should create a new algorithm if the input and output channels are non-zero.
-		// But we check the NumInputChannelsTemp/NumOutputChannelsTemp to check if we've
-		// already tried and failed to create an algorithm for this channel configuration.
-		bool bShouldCreateNewAlgo = (NumInputChannelsTemp != InData.NumChannels) || (NumOutputChannelsTemp != OutData.NumChannels);
-		bShouldCreateNewAlgo &= (InData.NumChannels > 0);
-		bShouldCreateNewAlgo &= (OutData.NumChannels > 0);
+	Reverb->ProcessAudio(InData.NumChannels, InData.AudioBuffer->GetData(), OutData.NumChannels, OutData.AudioBuffer->GetData(), InData.NumFrames);
 
-		if (bShouldCreateNewAlgo)
-		{
-			UE_LOG(LogSynthesis, Log, TEXT("Creating new convolution algorithm due to channel count update. Num Inputs %d -> %d. Num Outputs %d -> %d"), ExpectedNumInputChannels, InData.NumChannels, ExpectedNumOutputChannels, OutData.NumChannels);
-
-			// We don't update version data when changing the channel configuration
-			FVersionData CurrentVersionData;
-			{
-				FScopeLock VersionLock(&VersionDataCriticalSection);
-				CurrentVersionData = VersionData;
-			}
-
-			// These should only be updated here and in the constructor since we use them
-			// to see if we've tried this channel configuration already.
-			NumInputChannels.Store(InData.NumChannels);
-			NumOutputChannels.Store(OutData.NumChannels);
-
-			// Preferably, this should be done in a task. To do so, we would need 
-			// access to a weak ptr to "this". Currently that is not possible since this object
-			// is owned by the Preset, and the TWeakObjectPtr cannot safely be copied across thread
-			// boundaries. 
-			{
-				// Create the convolution algorithm init data
-				FConvolutionReverbInitData ConvolutionInitData = CreateConvolutionReverbInitData();
-
-				FConvolutionReverbSettings Settings;
-				Params.CopyParams(Settings);
-
-				// Create the convolution reverb
-				TUniquePtr<FConvolutionReverb> ConvReverb = FConvolutionReverb::CreateConvolutionReverb(ConvolutionInitData, Settings);
-
-				// Set the convolution algorithm, ensuring that it's the most up-to-date
-				SetConvolutionReverbIfCurrent(MoveTemp(ConvReverb), CurrentVersionData);
-			}
-
-			if (ConvolutionReverb.IsValid())
-			{
-				ExpectedNumInputChannels = ConvolutionReverb->GetNumInputChannels();
-				ExpectedNumOutputChannels = ConvolutionReverb->GetNumOutputChannels();
-			}
-		}
-	}
-
-	const bool bShouldProcessConvReverb = !bBypass &&
-		ConvolutionReverb.IsValid() && 
-		(ExpectedNumInputChannels == InData.NumChannels) && 
-		(ExpectedNumOutputChannels == OutData.NumChannels);
-
-	if (bShouldProcessConvReverb)
-	{
-		ConvolutionReverb->ProcessAudio(InData.NumChannels, *InData.AudioBuffer, OutData.NumChannels, *OutData.AudioBuffer);
-	}
-	else if (bBypass)
+	// Process Wet/Dry mix
+	if (FMath::IsNearlyEqual(WetVolume, 1.f) == false)
 	{
-		*OutData.AudioBuffer = *InData.AudioBuffer;
+		Audio::MultiplyBufferByConstantInPlace(*OutData.AudioBuffer, WetVolume);
 	}
-	else
+	if (DryVolume > KINDA_SMALL_NUMBER)
 	{
-		// Zero output data. Do *not* trigger rebuild here in case one is already in flight or simply because one cannot be built.
-		OutData.AudioBuffer->Reset();
-
-		int32 OutputNumFrames = InData.NumFrames * OutData.NumChannels;
-		if (OutputNumFrames > 0)
-		{
-			OutData.AudioBuffer->AddZeroed(OutputNumFrames);
-		}
+		Audio::MixInBufferFast(*InData.AudioBuffer, *OutData.AudioBuffer, DryVolume);
 	}
 }
 
-FSubmixEffectConvolutionReverb::FVersionData FSubmixEffectConvolutionReverb::UpdateVersion()
-{
-	FVersionData VersionDataCopy;
-	FScopeLock VersionDataLock(&VersionDataCriticalSection);
-
-	VersionData.ConvolutionID++;
-	VersionDataCopy = VersionData;
-
-	return VersionDataCopy;
-}
-
-void FSubmixEffectConvolutionReverb::UpdateParameters()
-{
-	// Call on the audio render thread to update parameters.
-	if (ConvolutionReverb.IsValid())
-	{
-		Audio::FConvolutionReverbSettings NewSettings;
-		if (Params.GetParams(&NewSettings))
-		{
-			ConvolutionReverb->SetSettings(NewSettings);
-		}
-    }
-}
-
-FSubmixEffectConvolutionReverb::FVersionData FSubmixEffectConvolutionReverb::UpdateConvolutionReverb(const USubmixEffectConvolutionReverbPreset* InPreset)
+AudioConvReverbIntrinsics::FVersionData FSubmixEffectConvolutionReverb::UpdateConvolutionReverb(const USubmixEffectConvolutionReverbPreset* InPreset)
 {
 	using namespace Audio;
 
 	// Copy data from preset into internal FConvolutionReverbInitData
-	
-	// TODO: Need to update AudioMixerSubmix initialization. Some initializations
-	// happen on the audio render thread.
-	// check(IsInGameThread() || IsInAudioThread());
-	{
-		FScopeLock ConvReverbInitDataLock(&ConvReverbInitDataCriticalSection);
-
-		// Reset data
-		ConvReverbInitData = FConvolutionReverbInitData();
+	// Reset data
+	FConvolutionReverbInitData InitData = FConvolutionReverbInitData();
 
-		if (nullptr != InPreset)
+	if (nullptr != InPreset)
+	{
+		if (InPreset->ImpulseResponse)
 		{
-			if (InPreset->ImpulseResponse)
-			{
-				UAudioImpulseResponse* IR = InPreset->ImpulseResponse;
-
-				ConvReverbInitData.Samples = IR->ImpulseResponse;
-				ConvReverbInitData.AlgorithmSettings.NumImpulseResponses = IR->NumChannels;
-				ConvReverbInitData.ImpulseSampleRate = IR->SampleRate;
-				ConvReverbInitData.bIsImpulseTrueStereo = IR->bTrueStereo && (IR->NumChannels % 2 == 0);
-				ConvReverbInitData.NormalizationVolume = Audio::ConvertToLinear(IR->NormalizationVolumeDb);
-			}
+			UAudioImpulseResponse* IR = InPreset->ImpulseResponse;
 
-			switch (InPreset->BlockSize)
-			{
-				case ESubmixEffectConvolutionReverbBlockSize::BlockSize256:
-					ConvReverbInitData.AlgorithmSettings.BlockNumSamples = 256;
-					break;
-
-				case ESubmixEffectConvolutionReverbBlockSize::BlockSize512:
-					ConvReverbInitData.AlgorithmSettings.BlockNumSamples = 512;
-					break;
-
-				case ESubmixEffectConvolutionReverbBlockSize::BlockSize1024:
-				default:
-					ConvReverbInitData.AlgorithmSettings.BlockNumSamples = 1024;
-					break;
-			}
+			InitData.Samples = IR->ImpulseResponse;
+			InitData.AlgorithmSettings.NumImpulseResponses = IR->NumChannels;
+			InitData.ImpulseSampleRate = IR->SampleRate;
+			InitData.bIsImpulseTrueStereo = IR->bTrueStereo && (IR->NumChannels % 2 == 0);
+			InitData.NormalizationVolume = ConvertToLinear(IR->NormalizationVolumeDb);
+		}
 
-			ConvReverbInitData.AlgorithmSettings.bEnableHardwareAcceleration = InPreset->bEnableHardwareAcceleration;
-			ConvReverbInitData.bMixInputChannelFormatToImpulseResponseFormat = InPreset->Settings.bMixInputChannelFormatToImpulseResponseFormat;
-			ConvReverbInitData.bMixReverbOutputToOutputChannelFormat = InPreset->Settings.bMixReverbOutputToOutputChannelFormat;
+		switch (InPreset->BlockSize)
+		{
+		case ESubmixEffectConvolutionReverbBlockSize::BlockSize256:
+			InitData.AlgorithmSettings.BlockNumSamples = 256;
+			break;
+
+		case ESubmixEffectConvolutionReverbBlockSize::BlockSize512:
+			InitData.AlgorithmSettings.BlockNumSamples = 512;
+			break;
+
+		case ESubmixEffectConvolutionReverbBlockSize::BlockSize1024:
+		default:
+			InitData.AlgorithmSettings.BlockNumSamples = 1024;
+			break;
 		}
+
+		InitData.AlgorithmSettings.bEnableHardwareAcceleration = InPreset->bEnableHardwareAcceleration;
+		InitData.bMixInputChannelFormatToImpulseResponseFormat = InPreset->Settings.bMixInputChannelFormatToImpulseResponseFormat;
+		InitData.bMixReverbOutputToOutputChannelFormat = InPreset->Settings.bMixReverbOutputToOutputChannelFormat;
 	}
 
-	return UpdateVersion();
+	Reverb->SetInitData(InitData);
+	return Reverb->UpdateVersion();
 }
 
-void FSubmixEffectConvolutionReverb::SetConvolutionReverbIfCurrent(TUniquePtr<Audio::FConvolutionReverb> InReverb, const FVersionData& InVersionData)
+void FSubmixEffectConvolutionReverb::RebuildConvolutionReverb()
 {
-	bool bIsCurrent = true;
-
-	{
-		FScopeLock VersionDataLock(&VersionDataCriticalSection);
-
-		bIsCurrent &= (InVersionData == VersionData);
-	}
-
-	if (bIsCurrent)
-	{
-		ConvolutionReverb = MoveTemp(InReverb);
-
-		if (ConvolutionReverb.IsValid())
-		{
-			ConvolutionReverb->SetSettings(GetParameters());
-		}
-	}
+	Reverb->BuildReverb();
 }
 
 /********************************************************************/
@@ -592,8 +250,9 @@ void USubmixEffectConvolutionReverbPreset::RebuildConvolutionReverb()
 	using namespace AudioConvReverbIntrinsics;
 	using FEffectWeakPtr = TWeakPtr<FSoundEffectBase, ESPMode::ThreadSafe>;
 	using FEffectSharedPtr = TSharedPtr<FSoundEffectBase, ESPMode::ThreadSafe>;
+	using FSubmixEffectConvSharedPtr = TSharedPtr<FSubmixEffectConvolutionReverb, ESPMode::ThreadSafe>;
 
-	// Iterator over effects and update the convolution reverb.
+	// Iterator over effects and update the convolution Reverb->
 	for (FEffectWeakPtr& InstanceWeakPtr : Instances)
 	{
 		FEffectSharedPtr InstanceSharedPtr = InstanceWeakPtr.Pin();
@@ -602,13 +261,7 @@ void USubmixEffectConvolutionReverbPreset::RebuildConvolutionReverb()
 			FSubmixEffectConvSharedPtr ConvEffectSharedPtr = StaticCastSharedPtr<FSubmixEffectConvolutionReverb>(InstanceSharedPtr);
 
 			FVersionData VersionData = ConvEffectSharedPtr->UpdateConvolutionReverb(this);
-
-			FConvolutionReverbInitData ConvolutionInitData = ConvEffectSharedPtr->CreateConvolutionReverbInitData();
-
-			Audio::FConvolutionReverbSettings ReverbParams = ConvEffectSharedPtr->GetParameters();
-
-			(new FAutoDeleteAsyncTask<FCreateConvolutionReverbTask>(ConvEffectSharedPtr, MoveTemp(ConvolutionInitData), ReverbParams, VersionData))->StartBackgroundTask();
-
+			ConvEffectSharedPtr->RebuildConvolutionReverb();
 		}
 	}
 }
diff --git a/Engine/Plugins/Runtime/Synthesis/Source/Synthesis/Public/ConvolutionReverb.h b/Engine/Plugins/Runtime/Synthesis/Source/Synthesis/Public/ConvolutionReverb.h
--- a/Engine/Plugins/Runtime/Synthesis/Source/Synthesis/Public/ConvolutionReverb.h
+++ b/Engine/Plugins/Runtime/Synthesis/Source/Synthesis/Public/ConvolutionReverb.h
@@ -165,7 +165,8 @@ namespace Audio
 
 
 		// If the number of input frames changes between callbacks, the output may contain discontinuities.
-		void ProcessAudio(int32 InNumInputChannels, FAlignedFloatBuffer& InputAudio, int32 InNumOutputChannels, FAlignedFloatBuffer& OutputAudio);
+		void ProcessAudio(int32 InNumInputChannels, const FAlignedFloatBuffer& InputAudio, int32 InNumOutputChannels, FAlignedFloatBuffer& OutputAudio);
+		void ProcessAudio(int32 InNumInputChannels, const float* InputAudio, int32 InNumOutputChannels, float* OutputAudio, const int32 InNumFrames);
 
 		int32 GetNumInputChannels() const;
 		int32 GetNumOutputChannels() const;
diff --git a/Engine/Plugins/Runtime/Synthesis/Source/Synthesis/Public/EffectConvolutionReverb.h b/Engine/Plugins/Runtime/Synthesis/Source/Synthesis/Public/EffectConvolutionReverb.h
new file mode 100644
--- /dev/null
+++ b/Engine/Plugins/Runtime/Synthesis/Source/Synthesis/Public/EffectConvolutionReverb.h
@@ -0,0 +1,148 @@
+// Copyright Epic Games, Inc. All Rights Reserved.
+#pragma once
+#include "CoreMinimal.h"
+#include "Async/Async.h"
+#include "ConvolutionReverb.h"
+#include "DSP/ConvolutionAlgorithm.h"
+#include "SynthesisModule.h"
+#include "EffectConvolutionReverb.generated.h"
+
+/** Block size of convolution reverb algorithm. */
+UENUM()
+enum class ESubmixEffectConvolutionReverbBlockSize : uint8
+{
+	/** 256 audio frames per a block. */
+	BlockSize256,
+
+	/** 512 audio frames per a block. */
+	BlockSize512,
+
+	/** 1024 audio frames per a block. */
+	BlockSize1024
+};
+
+#if WITH_EDITORONLY_DATA
+DECLARE_MULTICAST_DELEGATE_OneParam(FAudioImpulseResponsePropertyChange, FPropertyChangedEvent&);
+#endif
+
+// ========================================================================
+// UAudioImpulseResponse
+// UAsset used to represent Imported Impulse Responses
+// ========================================================================
+UCLASS(BlueprintType)
+class SYNTHESIS_API UAudioImpulseResponse : public UObject
+{
+	GENERATED_BODY()
+
+public:
+	UAudioImpulseResponse();
+	/** The interleaved audio samples used in convolution. */
+	UPROPERTY()
+	TArray<float> ImpulseResponse;
+
+	/** The number of channels in impulse response. */
+	UPROPERTY()
+	int32 NumChannels;
+
+	/** The original sample rate of the impulse response. */
+	UPROPERTY()
+	int32 SampleRate;
+
+	/* Used to account for energy added by convolution with "loud" Impulse Responses. */
+	UPROPERTY(EditAnywhere, BlueprintReadOnly, Category = SubmixEffectPreset, meta = (ClamMin = "-60.0", ClampMax = "15.0"))
+	float NormalizationVolumeDb;
+
+	/* If true, impulse response channels are interpreted as true stereo which allows channel crosstalk. If false, impulse response channels are interpreted as independent channel impulses. */
+	UPROPERTY(EditAnywhere, BlueprintReadOnly, Category = SubmixEffectPreset, meta = (EditCondition = "(NumChannels > 0) && (NumChannels % 2 == 0)"))
+	bool bTrueStereo;
+
+	UPROPERTY(meta = (DeprecatedProperty))
+	TArray<float> IRData_DEPRECATED;
+
+#if WITH_EDITORONLY_DATA
+	virtual void PostEditChangeProperty(FPropertyChangedEvent& PropertyChangedEvent) override;
+
+	// This delegate is called whenever PostEditChangeProperty called.
+	FAudioImpulseResponsePropertyChange OnObjectPropertyChanged;
+#endif
+};
+
+namespace AudioConvReverbIntrinsics
+{
+	using namespace Audio;
+
+	typedef int32 FConvolutionReverbID;
+
+	struct FVersionData
+	{
+		FConvolutionReverbID ConvolutionID;
+
+		FVersionData()
+		{
+			ConvolutionID = 0;
+		}
+
+		bool operator==(const FVersionData& Other) const
+		{
+			// Check if versions are equal
+			return ConvolutionID == Other.ConvolutionID;
+		}
+	};
+}
+
+namespace Audio
+{
+	class SYNTHESIS_API FEffectConvolutionReverb : public TSharedFromThis<FEffectConvolutionReverb, ESPMode::ThreadSafe>
+	{
+	public:
+
+		void ProcessAudio(int32 InNumInputChannels, const float* InputAudio, int32 InNumOutputChannels, float* OutputAudio, const int32 InNumFrames);
+
+		FConvolutionReverbInitData CreateConvolutionReverbInitData();
+
+		// Fire off an async task to build a new reverb algorithm with new settings
+		void BuildReverb();
+
+		// Immediately build the initial reverb algorithm
+		void Init();
+
+		void SetSettings(const FConvolutionReverbSettings& InSettings);
+		void SetSampleRate(const float InSampleRate);
+		void SetInitData(const FConvolutionReverbInitData& InData);
+		void SetReverb(TUniquePtr<FConvolutionReverb> InReverb);
+		void SetBypass(const bool InBypass);
+
+		AudioConvReverbIntrinsics::FVersionData UpdateVersion();
+		void UpdateChannelCount(const int32 InNumInputChannels, const int32 InNumOutputChannels);
+		void UpdateParameters();
+
+		void EnqueueNewReverb(TUniquePtr<FConvolutionReverb> InReverb, const AudioConvReverbIntrinsics::FVersionData& InVersionData);
+		void DequeueNewReverb();
+
+		bool IsVersionCurrent(const AudioConvReverbIntrinsics::FVersionData& InVersionData) const;
+		bool IsChannelCountUpToDate(const int32 InNumInputChannels, const int32 InNumOutputChannels) const;
+
+	private:
+		TUniquePtr<FConvolutionReverb> Reverb;
+
+		// Reverb created by our async builder
+		TUniquePtr<FConvolutionReverb> QueuedReverb;
+
+		mutable FCriticalSection VersionDataCriticalSection;
+		AudioConvReverbIntrinsics::FVersionData VersionData;
+		AudioConvReverbIntrinsics::FVersionData QueuedVersionData;
+
+		mutable FCriticalSection ConvReverbInitDataCriticalSection;
+		// Internal copy of data needed to create FConvolutionReverb
+		FConvolutionReverbInitData ConvReverbInitData;
+
+		TAtomic<int32> NumInputChannels = 2;
+		TAtomic<int32> NumOutputChannels = 2;
+
+		// Params object for copying between audio render thread and outside threads.
+		TParams<FConvolutionReverbSettings> Params;
+
+		bool bBypass = false;
+		float SampleRate = 0.f;
+	};
+}
