Adding copyright to EffectConvolutionReverb.cpp

diff --git a/Engine/Plugins/Runtime/Synthesis/Source/Synthesis/Private/EffectConvolutionReverb.cpp b/Engine/Plugins/Runtime/Synthesis/Source/Synthesis/Private/EffectConvolutionReverb.cpp
--- a/Engine/Plugins/Runtime/Synthesis/Source/Synthesis/Private/EffectConvolutionReverb.cpp
+++ b/Engine/Plugins/Runtime/Synthesis/Source/Synthesis/Private/EffectConvolutionReverb.cpp
@@ -1,3 +1,5 @@
+// Copyright Epic Games, Inc. All Rights Reserved.
+
 #include "EffectConvolutionReverb.h"
 
 UAudioImpulseResponse::UAudioImpulseResponse()
