Add include to EffectConvolutionReverb.h to fix android non-unity build fix

diff --git a/Engine/Plugins/Runtime/Synthesis/Source/Synthesis/Public/EffectConvolutionReverb.h b/Engine/Plugins/Runtime/Synthesis/Source/Synthesis/Public/EffectConvolutionReverb.h
--- a/Engine/Plugins/Runtime/Synthesis/Source/Synthesis/Public/EffectConvolutionReverb.h
+++ b/Engine/Plugins/Runtime/Synthesis/Source/Synthesis/Public/EffectConvolutionReverb.h
@@ -4,6 +4,7 @@
 #include "Async/Async.h"
 #include "ConvolutionReverb.h"
 #include "DSP/ConvolutionAlgorithm.h"
+#include "DSP/Dsp.h"
 #include "SynthesisModule.h"
 #include "EffectConvolutionReverb.generated.h"
 
