Fix for multiband compressor crash on mac due to unaligned buffers
#rnx

diff --git a/Engine/Plugins/Runtime/Synthesis/Source/Synthesis/Classes/SubmixEffects/SubmixEffectMultiBandCompressor.h b/Engine/Plugins/Runtime/Synthesis/Source/Synthesis/Classes/SubmixEffects/SubmixEffectMultiBandCompressor.h
--- a/Engine/Plugins/Runtime/Synthesis/Source/Synthesis/Classes/SubmixEffects/SubmixEffectMultiBandCompressor.h
+++ b/Engine/Plugins/Runtime/Synthesis/Source/Synthesis/Classes/SubmixEffects/SubmixEffectMultiBandCompressor.h
@@ -140,6 +140,8 @@ protected:
 	Audio::FAlignedFloatBuffer AudioExternal;
 
 private:
+	static constexpr int32 MaxBlockNumSamples = 8192; // 8 channels * 1024 samples
+	static_assert((MaxBlockNumSamples % 4) == 0, "MaxBlockNumSamples must be evenly divisible by 4 to allow for SIMD optimization" );
 	FKeySource KeySource;
 
 	Audio::FDeviceId DeviceId;
@@ -147,6 +149,8 @@ private:
 	FDelegateHandle DeviceCreatedHandle;
 	FDelegateHandle DeviceDestroyedHandle;
 
+	Audio::FAlignedFloatBuffer ScratchBuffer;
+
 	TArray<Audio::FDynamicsProcessor> DynamicsProcessors;
 
 	Audio::FLinkwitzRileyBandSplitter KeyBandSplitter;
@@ -205,4 +209,4 @@ public:
 
 private:
 	void SetKey(ESubmixEffectDynamicsKeySource InKeySource, UObject* InObject, int32 InNumChannels = 0);
-};
\ No newline at end of file
+};
diff --git a/Engine/Plugins/Runtime/Synthesis/Source/Synthesis/Private/SubmixEffectMultiBandCompressor.cpp b/Engine/Plugins/Runtime/Synthesis/Source/Synthesis/Private/SubmixEffectMultiBandCompressor.cpp
--- a/Engine/Plugins/Runtime/Synthesis/Source/Synthesis/Private/SubmixEffectMultiBandCompressor.cpp
+++ b/Engine/Plugins/Runtime/Synthesis/Source/Synthesis/Private/SubmixEffectMultiBandCompressor.cpp
@@ -9,6 +9,8 @@
 
 FSubmixEffectMultibandCompressor::FSubmixEffectMultibandCompressor()
 {
+	ScratchBuffer.AddUninitialized(FSubmixEffectMultibandCompressor::MaxBlockNumSamples);
+
 	DeviceCreatedHandle = FAudioDeviceManagerDelegates::OnAudioDeviceCreated.AddRaw(this, &FSubmixEffectMultibandCompressor::OnDeviceCreated);
 	DeviceDestroyedHandle = FAudioDeviceManagerDelegates::OnAudioDeviceDestroyed.AddRaw(this, &FSubmixEffectMultibandCompressor::OnDeviceDestroyed);
 }
@@ -386,10 +388,8 @@ void FSubmixEffectMultibandCompressor::OnProcessAudio(const FSoundEffectSubmixIn
 	// Zero buffer so multiple bands can all sum to it
 	FMemory::Memzero(OutBuffer.GetData(), FrameSize * InData.NumFrames);
 
-	Audio::FStackSampleBuffer ScratchBuffer;
-
-	// Scratch buffer is stack allocated with a max size. If passed more samples than it can hold, process in blocks
-	const int32 BlockSize = FMath::Min(ScratchBuffer.Max(), NumSamples);
+	// If passed more samples than the max block size, process in blocks
+	const int32 BlockSize = FMath::Min(MaxBlockNumSamples, NumSamples);
 
 	if (BlockSize > MultiBandBuffer.NumSamples)
 	{
@@ -416,7 +416,7 @@ void FSubmixEffectMultibandCompressor::OnProcessAudio(const FSoundEffectSubmixIn
 			KeyPtr = &AudioExternal[KeySampleIdx];
 		}
 
-		ScratchBuffer.SetNumUninitialized(BlockSize);
+		ScratchBuffer.SetNumUninitialized(BlockSize, false /* bAllowShrinking */);
 
 		BandSplitter.ProcessAudioBuffer(InPtr, MultiBandBuffer, NumBlockFrames);
 
diff --git a/Engine/Source/Runtime/SignalProcessing/Private/LinkwitzRileyBandSplitter.cpp b/Engine/Source/Runtime/SignalProcessing/Private/LinkwitzRileyBandSplitter.cpp
--- a/Engine/Source/Runtime/SignalProcessing/Private/LinkwitzRileyBandSplitter.cpp
+++ b/Engine/Source/Runtime/SignalProcessing/Private/LinkwitzRileyBandSplitter.cpp
@@ -149,14 +149,11 @@ namespace Audio
 			return;
 		}
 
-		FStackSampleBuffer SharedStackBuffer;
-		SharedStackBuffer.SetNumZeroed(NumSamples);
+		SharedAlignedBuffer.SetNumZeroed(NumSamples, false /* bAllowShrinking */);
+		BandAlignedBuffer.SetNumUninitialized(NumSamples, false /* bAllowShrinking */);
 
-		FStackSampleBuffer BandStackBuffer;
-		BandStackBuffer.SetNumUninitialized(NumSamples);
-
-		float* const SharedBufferPtr = SharedStackBuffer.GetData();
-		float* const BandBufferPtr = BandStackBuffer.GetData();
+		float* const SharedBufferPtr = SharedAlignedBuffer.GetData();
+		float* const BandBufferPtr = BandAlignedBuffer.GetData();
 
 		CopyToBuffer(SharedBufferPtr, InBuffer, NumSamples);
 
@@ -271,4 +268,4 @@ namespace Audio
 
 		return Audio::ConvertToLinear(Qdb);
 	}
-}
\ No newline at end of file
+}
diff --git a/Engine/Source/Runtime/SignalProcessing/Public/DSP/LinkwitzRileyBandSplitter.h b/Engine/Source/Runtime/SignalProcessing/Public/DSP/LinkwitzRileyBandSplitter.h
--- a/Engine/Source/Runtime/SignalProcessing/Public/DSP/LinkwitzRileyBandSplitter.h
+++ b/Engine/Source/Runtime/SignalProcessing/Public/DSP/LinkwitzRileyBandSplitter.h
@@ -11,8 +11,6 @@
 
 namespace Audio
 {
-	typedef TArray<float, TInlineAllocator<MAX_BUFFER_SIZE>> FStackSampleBuffer;
-
 	struct FLinkwitzRileyBandFilter
 	{
 		TArray<FVariablePoleFilter> Filters;
@@ -102,6 +100,8 @@ namespace Audio
 
 		TArray<float> SharedBuffer;
 		TArray<float> BandWorkBuffer;
+		FAlignedFloatBuffer SharedAlignedBuffer;
+		FAlignedFloatBuffer BandAlignedBuffer;
 
 		TArray<FLinkwitzRileyBandFilter> BandFilters;
 		TArray<FCrossoverBandwidthPair> Crossovers;
@@ -110,4 +110,4 @@ namespace Audio
 		void InvertBuffer(float* Buffer, const int32 NumSamples);
 		float GetQ(EFilterOrder InFilterOrder);
 	};
-}
\ No newline at end of file
+}
