Add data to spat params to allow using the AttenuationDistance (which may differ from distance to listener)

- Work done specifically for SourceEffectSimpleDelay to use it for "DelaySpeedBasedOnDistance" calculations

Impact: Prevents pitch warble when moving the camera around quickly for projects that use an attenuation override

diff --git a/Engine/Plugins/Runtime/Synthesis/Source/Synthesis/Classes/SourceEffects/SourceEffectSimpleDelay.h b/Engine/Plugins/Runtime/Synthesis/Source/Synthesis/Classes/SourceEffects/SourceEffectSimpleDelay.h
--- a/Engine/Plugins/Runtime/Synthesis/Source/Synthesis/Classes/SourceEffects/SourceEffectSimpleDelay.h
+++ b/Engine/Plugins/Runtime/Synthesis/Source/Synthesis/Classes/SourceEffects/SourceEffectSimpleDelay.h
@@ -35,6 +35,10 @@ struct SYNTHESIS_API FSourceEffectSimpleDelaySettings
 	UPROPERTY(EditAnywhere, BlueprintReadWrite, Category = "SourceEffect|Preset")
 	uint32 bDelayBasedOnDistance : 1;
 
+	// Whether or not to allow the attenuation distance override value vs the distance to listener to be used for distance-based delay.
+	UPROPERTY(EditAnywhere, BlueprintReadWrite, Category = "SourceEffect|Preset")
+	uint32 bUseDistanceOverride : 1;
+	
 	FSourceEffectSimpleDelaySettings()
 		: SpeedOfSound(343.0f)
 		, DelayAmount(0.0f)
@@ -42,6 +46,7 @@ struct SYNTHESIS_API FSourceEffectSimpleDelaySettings
 		, WetAmount(1.0f)
 		, Feedback(0.0f)
 		, bDelayBasedOnDistance(true)
+		, bUseDistanceOverride(true)
 	{}
 };
 
diff --git a/Engine/Plugins/Runtime/Synthesis/Source/Synthesis/Private/SourceEffectSimpleDelay.cpp b/Engine/Plugins/Runtime/Synthesis/Source/Synthesis/Private/SourceEffectSimpleDelay.cpp
--- a/Engine/Plugins/Runtime/Synthesis/Source/Synthesis/Private/SourceEffectSimpleDelay.cpp
+++ b/Engine/Plugins/Runtime/Synthesis/Source/Synthesis/Private/SourceEffectSimpleDelay.cpp
@@ -37,7 +37,10 @@ void FSourceEffectSimpleDelay::ProcessAudio(const FSoundEffectSourceInputData& I
 {
 	if (SettingsCopy.bDelayBasedOnDistance)
 	{
-		const float DistanceMeters = InData.SpatParams.Distance * 0.01f;
+		const float DistanceMeters = SettingsCopy.bUseDistanceOverride ?
+			InData.SpatParams.AttenuationDistance * 0.1f
+			: InData.SpatParams.Distance * 0.01f;
+
 		const float DelayAmountMsec = 1000.0f * DistanceMeters / SettingsCopy.SpeedOfSound;
 
 		for (Audio::FDelay& Delay : Delays)
diff --git a/Engine/Source/Runtime/AudioExtensions/Public/IAudioExtensionPlugin.h b/Engine/Source/Runtime/AudioExtensions/Public/IAudioExtensionPlugin.h
--- a/Engine/Source/Runtime/AudioExtensions/Public/IAudioExtensionPlugin.h
+++ b/Engine/Source/Runtime/AudioExtensions/Public/IAudioExtensionPlugin.h
@@ -91,6 +91,9 @@ struct FSpatializationParams
 	/** The distance between listener and emitter. */
 	float Distance;
 
+	/** The distance used to compute attenuation. Maybe different from the distance between listener and emitter if it's overridden. */
+	float AttenuationDistance;
+
 	/** The normalized omni radius, or the radius that will blend a sound to non-3d */
 	float NormalizedOmniRadius;
 
@@ -107,6 +110,7 @@ struct FSpatializationParams
 		, LeftChannelPosition(FVector::ZeroVector)
 		, RightChannelPosition(FVector::ZeroVector)
 		, Distance(0.0f)
+		, AttenuationDistance(0.0f)
 		, NormalizedOmniRadius(0.0f)
 		, AudioClock(0.0)
 	{}
diff --git a/Engine/Source/Runtime/Engine/Private/Audio.cpp b/Engine/Source/Runtime/Engine/Private/Audio.cpp
--- a/Engine/Source/Runtime/Engine/Private/Audio.cpp
+++ b/Engine/Source/Runtime/Engine/Private/Audio.cpp
@@ -538,6 +538,9 @@ FSpatializationParams FSoundSource::GetSpatializationParams()
 	{
 		FVector EmitterPosition = AudioDevice->GetListenerTransformedDirection(WaveInstance->Location, &Params.Distance);
 
+		// Independently retrieve the attenuation used for distance in case it was overridden
+		Params.AttenuationDistance = AudioDevice->GetDistanceToNearestListener(WaveInstance->Location);
+		
 		// If we are using the OmniRadius feature
 		if (WaveInstance->OmniRadius > 0.0f)
 		{
