Audio Plugin Telemetry Part 5
Steam Audio, Synthesis & DSP Effects, and AudioLink

diff --git a/Engine/Plugins/Runtime/Steam/SteamAudio/Source/SteamAudioEditor/Private/PhononOcclusionSettingsFactory.cpp b/Engine/Plugins/Runtime/Steam/SteamAudio/Source/SteamAudioEditor/Private/PhononOcclusionSettingsFactory.cpp
--- a/Engine/Plugins/Runtime/Steam/SteamAudio/Source/SteamAudioEditor/Private/PhononOcclusionSettingsFactory.cpp
+++ b/Engine/Plugins/Runtime/Steam/SteamAudio/Source/SteamAudioEditor/Private/PhononOcclusionSettingsFactory.cpp
@@ -3,6 +3,8 @@
 //
 
 #include "PhononOcclusionSettingsFactory.h"
+
+#include "EngineAnalytics.h"
 #include "PhononOcclusionSourceSettings.h"
 #include "UObject/ObjectMacros.h"
 #include "UObject/Object.h"
@@ -52,6 +54,10 @@ UPhononOcclusionSettingsFactory::UPhononOcclusionSettingsFactory(const FObjectIn
 UObject* UPhononOcclusionSettingsFactory::FactoryCreateNew(UClass* Class, UObject* InParent, FName InName, EObjectFlags Flags, UObject* Context,
 	FFeedbackContext* Warn)
 {
+	if (FEngineAnalytics::IsAvailable())
+	{
+		FEngineAnalytics::GetProvider().RecordEvent(TEXT("Audio.Usage.SteamAudio.PhononOcclusionSettingsCreated"));
+	}
 	return NewObject<UPhononOcclusionSourceSettings>(InParent, InName, Flags);
 }
 
diff --git a/Engine/Plugins/Runtime/Steam/SteamAudio/Source/SteamAudioEditor/Private/PhononReverbSettingsFactory.cpp b/Engine/Plugins/Runtime/Steam/SteamAudio/Source/SteamAudioEditor/Private/PhononReverbSettingsFactory.cpp
--- a/Engine/Plugins/Runtime/Steam/SteamAudio/Source/SteamAudioEditor/Private/PhononReverbSettingsFactory.cpp
+++ b/Engine/Plugins/Runtime/Steam/SteamAudio/Source/SteamAudioEditor/Private/PhononReverbSettingsFactory.cpp
@@ -3,6 +3,8 @@
 //
 
 #include "PhononReverbSettingsFactory.h"
+
+#include "EngineAnalytics.h"
 #include "PhononReverbSourceSettings.h"
 #include "UObject/ObjectMacros.h"
 #include "UObject/Object.h"
@@ -52,6 +54,10 @@ UPhononReverbSettingsFactory::UPhononReverbSettingsFactory(const FObjectInitiali
 UObject* UPhononReverbSettingsFactory::FactoryCreateNew(UClass* Class, UObject* InParent, FName InName, EObjectFlags Flags, UObject* Context,
 	FFeedbackContext* Warn)
 {
+	if (FEngineAnalytics::IsAvailable())
+	{
+		FEngineAnalytics::GetProvider().RecordEvent(TEXT("Audio.Usage.SteamAudio.PhononReverbSettingsCreated"));
+	}
 	return NewObject<UPhononReverbSourceSettings>(InParent, InName, Flags);
 }
 
diff --git a/Engine/Plugins/Runtime/Steam/SteamAudio/Source/SteamAudioEditor/Private/PhononSpatializationSettingsFactory.cpp b/Engine/Plugins/Runtime/Steam/SteamAudio/Source/SteamAudioEditor/Private/PhononSpatializationSettingsFactory.cpp
--- a/Engine/Plugins/Runtime/Steam/SteamAudio/Source/SteamAudioEditor/Private/PhononSpatializationSettingsFactory.cpp
+++ b/Engine/Plugins/Runtime/Steam/SteamAudio/Source/SteamAudioEditor/Private/PhononSpatializationSettingsFactory.cpp
@@ -3,6 +3,8 @@
 //
 
 #include "PhononSpatializationSettingsFactory.h"
+
+#include "EngineAnalytics.h"
 #include "PhononSpatializationSourceSettings.h"
 #include "UObject/ObjectMacros.h"
 #include "UObject/Object.h"
@@ -52,6 +54,10 @@ UPhononSpatializationSettingsFactory::UPhononSpatializationSettingsFactory(const
 UObject* UPhononSpatializationSettingsFactory::FactoryCreateNew(UClass* Class, UObject* InParent, FName InName, EObjectFlags Flags,
 	UObject* Context, FFeedbackContext* Warn)
 {
+	if (FEngineAnalytics::IsAvailable())
+	{
+		FEngineAnalytics::GetProvider().RecordEvent(TEXT("Audio.Usage.SteamAudio.PhononSpatializationSettingsCreated"));
+	}
 	return NewObject<UPhononSpatializationSourceSettings>(InParent, InName, Flags);
 }
 
diff --git a/Engine/Plugins/Runtime/Synthesis/Source/SynthesisEditor/Private/EpicSynth1PresetBank.cpp b/Engine/Plugins/Runtime/Synthesis/Source/SynthesisEditor/Private/EpicSynth1PresetBank.cpp
--- a/Engine/Plugins/Runtime/Synthesis/Source/SynthesisEditor/Private/EpicSynth1PresetBank.cpp
+++ b/Engine/Plugins/Runtime/Synthesis/Source/SynthesisEditor/Private/EpicSynth1PresetBank.cpp
@@ -1,6 +1,7 @@
 // Copyright Epic Games, Inc. All Rights Reserved.
 #include "EpicSynth1PresetBank.h"
 
+#include "EngineAnalytics.h"
 #include "SynthComponents/EpicSynth1Component.h"
 #include "UObject/ObjectMacros.h"
 #include "UObject/Object.h"
@@ -36,7 +37,10 @@ UModularSynthPresetBankFactory::UModularSynthPresetBankFactory(const FObjectInit
 UObject* UModularSynthPresetBankFactory::FactoryCreateNew(UClass* Class, UObject* InParent, FName InName, EObjectFlags Flags, UObject* Context, FFeedbackContext* Warn)
 {
 	UModularSynthPresetBank* NewPresetBank = NewObject<UModularSynthPresetBank>(InParent, InName, Flags);
-
+	if (FEngineAnalytics::IsAvailable())
+	{
+		FEngineAnalytics::GetProvider().RecordEvent(TEXT("Audio.Usage.SynthesisAndDSPEffects.ModularSynthPresetBankCreated"));
+	}
 	return NewPresetBank;
 }
 
diff --git a/Engine/Plugins/Runtime/Synthesis/Source/SynthesisEditor/Private/MonoWaveTablePresetBank.cpp b/Engine/Plugins/Runtime/Synthesis/Source/SynthesisEditor/Private/MonoWaveTablePresetBank.cpp
--- a/Engine/Plugins/Runtime/Synthesis/Source/SynthesisEditor/Private/MonoWaveTablePresetBank.cpp
+++ b/Engine/Plugins/Runtime/Synthesis/Source/SynthesisEditor/Private/MonoWaveTablePresetBank.cpp
@@ -1,6 +1,7 @@
 // Copyright Epic Games, Inc. All Rights Reserved.
 #include "MonoWaveTablePresetBank.h"
 
+#include "EngineAnalytics.h"
 #include "SynthComponents/SynthComponentMonoWaveTable.h"
 #include "UObject/ObjectMacros.h"
 #include "UObject/Object.h"
@@ -36,7 +37,10 @@ UMonoWaveTableSynthPresetFactory::UMonoWaveTableSynthPresetFactory(const FObject
 UObject* UMonoWaveTableSynthPresetFactory::FactoryCreateNew(UClass* Class, UObject* InParent, FName InName, EObjectFlags Flags, UObject* Context, FFeedbackContext* Warn)
 {
 	UMonoWaveTableSynthPreset* NewPresetBank = NewObject<UMonoWaveTableSynthPreset>(InParent, InName, Flags);
-
+	if (FEngineAnalytics::IsAvailable())
+	{
+		FEngineAnalytics::GetProvider().RecordEvent(TEXT("Audio.Usage.SynthesisAndDSPEffects.MonoWaveTableSynthPresetCreated"));
+	}
 	return NewPresetBank;
 }
 #undef LOCTEXT_NAMESPACE
\ No newline at end of file
diff --git a/Engine/Source/Runtime/AudioLink/AudioLinkEngine/Private/IAudioLink.cpp b/Engine/Source/Runtime/AudioLink/AudioLinkEngine/Private/IAudioLink.cpp
new file mode 100644
--- /dev/null
+++ b/Engine/Source/Runtime/AudioLink/AudioLinkEngine/Private/IAudioLink.cpp
@@ -0,0 +1,14 @@
+// Copyright Epic Games, Inc. All Rights Reserved.
+
+#include "IAudioLink.h"
+
+#include "EngineAnalytics.h"
+
+IAudioLink::IAudioLink()
+{
+	if (FEngineAnalytics::IsAvailable())
+	{
+		FEngineAnalytics::GetProvider().RecordEvent(TEXT("Audio.Usage.AudioLink.InstanceCreated"));
+	}
+}
+
diff --git a/Engine/Source/Runtime/AudioLink/AudioLinkEngine/Public/IAudioLink.h b/Engine/Source/Runtime/AudioLink/AudioLinkEngine/Public/IAudioLink.h
--- a/Engine/Source/Runtime/AudioLink/AudioLinkEngine/Public/IAudioLink.h
+++ b/Engine/Source/Runtime/AudioLink/AudioLinkEngine/Public/IAudioLink.h
@@ -2,8 +2,6 @@
 
 #pragma once
 
-#include "CoreMinimal.h"
-
 /** IAudioLink
   *	Abstract interface for AudioLink instances. 
   *	Purely opaque.
@@ -11,7 +9,7 @@
 class IAudioLink
 {
 protected:
-	IAudioLink() = default;
+	IAudioLink();
 public:
 	virtual ~IAudioLink() = default;
 };
