Fix for convolution reverb not working on higher channel-count AudioDevices / SubmixGraphs

diff --git a/Engine/Plugins/Runtime/Synthesis/Source/Synthesis/Private/EffectConvolutionReverb.cpp b/Engine/Plugins/Runtime/Synthesis/Source/Synthesis/Private/EffectConvolutionReverb.cpp
--- a/Engine/Plugins/Runtime/Synthesis/Source/Synthesis/Private/EffectConvolutionReverb.cpp
+++ b/Engine/Plugins/Runtime/Synthesis/Source/Synthesis/Private/EffectConvolutionReverb.cpp
@@ -260,7 +260,7 @@ namespace Audio
 		}
 
 		// Check if there is a channel mismatch between the convolution reverb and the input/output data
-		const bool bNumChannelsMismatch = IsChannelCountUpToDate(InNumInputChannels, InNumOutputChannels);
+		const bool bNumChannelsMismatch = !IsChannelCountUpToDate(InNumInputChannels, InNumOutputChannels);
 		const bool bRequestedChannelCountIsValid = InNumInputChannels > 0 && InNumOutputChannels > 0;
 
 		if (bNumChannelsMismatch && bRequestedChannelCountIsValid)
