[Granular synth] protect against loading USoundWaves w/ incompatible loading behavior.
#rnx

diff --git a/Engine/Plugins/Runtime/Synthesis/Source/Synthesis/Private/SynthComponentGranulator.cpp b/Engine/Plugins/Runtime/Synthesis/Source/Synthesis/Private/SynthComponentGranulator.cpp
--- a/Engine/Plugins/Runtime/Synthesis/Source/Synthesis/Private/SynthComponentGranulator.cpp
+++ b/Engine/Plugins/Runtime/Synthesis/Source/Synthesis/Private/SynthComponentGranulator.cpp
@@ -114,7 +114,7 @@ void UGranularSynth::SetSoundWave(USoundWave* InSoundWave)
 		GranulatedSoundWave = InSoundWave;
 		bIsLoaded = false;
 
-		if (InSoundWave)
+		if (InSoundWave && InSoundWave->GetLoadingBehavior() == ESoundWaveLoadingBehavior::ForceInline)
 		{
 			TFunction<void(const USoundWave * SoundWave, const Audio::FSampleBuffer & SampleBuffer)> OnLoaded
 				= [this](const USoundWave * SoundWave, const Audio::FSampleBuffer & SampleBuffer)
