Follow up fix for guard around Granular Synth being provided a USoundWave w/ non-forceinline loading behavior.
# jira UE-120998
#rnx

diff --git a/Engine/Plugins/Runtime/Synthesis/Source/Synthesis/Private/SynthComponentGranulator.cpp b/Engine/Plugins/Runtime/Synthesis/Source/Synthesis/Private/SynthComponentGranulator.cpp
--- a/Engine/Plugins/Runtime/Synthesis/Source/Synthesis/Private/SynthComponentGranulator.cpp
+++ b/Engine/Plugins/Runtime/Synthesis/Source/Synthesis/Private/SynthComponentGranulator.cpp
@@ -114,7 +114,7 @@ void UGranularSynth::SetSoundWave(USoundWave* InSoundWave)
 		GranulatedSoundWave = InSoundWave;
 		bIsLoaded = false;
 
-		if (InSoundWave && InSoundWave->GetLoadingBehavior() == ESoundWaveLoadingBehavior::ForceInline)
+		if (InSoundWave && InSoundWave->GetLoadingBehavior(/*bCheckSoundClasses*/ false) == ESoundWaveLoadingBehavior::ForceInline)
 		{
 			TFunction<void(const USoundWave * SoundWave, const Audio::FSampleBuffer & SampleBuffer)> OnLoaded
 				= [this](const USoundWave * SoundWave, const Audio::FSampleBuffer & SampleBuffer)
