Fix EditCondition warning in Impulse Response Asset

diff --git a/Engine/Plugins/Runtime/Synthesis/Source/Synthesis/Private/EffectConvolutionReverb.cpp b/Engine/Plugins/Runtime/Synthesis/Source/Synthesis/Private/EffectConvolutionReverb.cpp
--- a/Engine/Plugins/Runtime/Synthesis/Source/Synthesis/Private/EffectConvolutionReverb.cpp
+++ b/Engine/Plugins/Runtime/Synthesis/Source/Synthesis/Private/EffectConvolutionReverb.cpp
@@ -21,6 +21,18 @@ void UAudioImpulseResponse::PostEditChangeProperty(FPropertyChangedEvent& Proper
 }
 #endif
 
+void UAudioImpulseResponse::Serialize(FArchive& Ar)
+{
+	Super::Serialize(Ar);
+
+#if WITH_EDITORONLY_DATA
+	if (Ar.IsLoading())
+	{
+		bIsEvenChannelCount = NumChannels % 2 == 0;
+	}
+#endif // WITH_EDITORONLY_DATA
+}
+
 namespace AudioConvReverbIntrinsics
 {
 	// Task for creating convolution algorithm object.
diff --git a/Engine/Plugins/Runtime/Synthesis/Source/Synthesis/Public/EffectConvolutionReverb.h b/Engine/Plugins/Runtime/Synthesis/Source/Synthesis/Public/EffectConvolutionReverb.h
--- a/Engine/Plugins/Runtime/Synthesis/Source/Synthesis/Public/EffectConvolutionReverb.h
+++ b/Engine/Plugins/Runtime/Synthesis/Source/Synthesis/Public/EffectConvolutionReverb.h
@@ -54,17 +54,23 @@ public:
 	float NormalizationVolumeDb;
 
 	/* If true, impulse response channels are interpreted as true stereo which allows channel crosstalk. If false, impulse response channels are interpreted as independent channel impulses. */
-	UPROPERTY(EditAnywhere, BlueprintReadOnly, Category = SubmixEffectPreset, meta = (EditCondition = "(NumChannels > 0) && (NumChannels % 2 == 0)"))
+	UPROPERTY(EditAnywhere, BlueprintReadOnly, Category = SubmixEffectPreset, meta = (EditCondition = "(NumChannels > 0) && bIsEvenChannelCount"))
 	bool bTrueStereo;
 
 	UPROPERTY(meta = (DeprecatedProperty))
 	TArray<float> IRData_DEPRECATED;
 
+	virtual void Serialize(FArchive& Ar) override;
+
 #if WITH_EDITORONLY_DATA
 	virtual void PostEditChangeProperty(FPropertyChangedEvent& PropertyChangedEvent) override;
 
 	// This delegate is called whenever PostEditChangeProperty called.
 	FAudioImpulseResponsePropertyChange OnObjectPropertyChanged;
+
+	// Used to evaluate whether TrueStereo should be editable
+	UPROPERTY()
+	bool bIsEvenChannelCount;
 #endif
 };
 
diff --git a/Engine/Plugins/Runtime/Synthesis/Source/SynthesisEditor/Private/AudioImpulseResponseAsset.cpp b/Engine/Plugins/Runtime/Synthesis/Source/SynthesisEditor/Private/AudioImpulseResponseAsset.cpp
--- a/Engine/Plugins/Runtime/Synthesis/Source/SynthesisEditor/Private/AudioImpulseResponseAsset.cpp
+++ b/Engine/Plugins/Runtime/Synthesis/Source/SynthesisEditor/Private/AudioImpulseResponseAsset.cpp
@@ -111,6 +111,10 @@ UObject* UAudioImpulseResponseFactory::FactoryCreateNew(UClass* Class, UObject*
 			}
 		}
 
+#if WITH_EDITORONLY_DATA
+		NewAsset->bIsEvenChannelCount = NewAsset->NumChannels % 2 == 0;
+#endif
+
 		StagedSoundWave.Reset();
 	}
 
