Prevent crash in multiband compressor when it tries to run with 0 bands.

#rnx

diff --git a/Engine/Plugins/Runtime/Synthesis/Source/Synthesis/Private/SubmixEffectMultiBandCompressor.cpp b/Engine/Plugins/Runtime/Synthesis/Source/Synthesis/Private/SubmixEffectMultiBandCompressor.cpp
--- a/Engine/Plugins/Runtime/Synthesis/Source/Synthesis/Private/SubmixEffectMultiBandCompressor.cpp
+++ b/Engine/Plugins/Runtime/Synthesis/Source/Synthesis/Private/SubmixEffectMultiBandCompressor.cpp
@@ -296,6 +296,15 @@ bool FSubmixEffectMultibandCompressor::UpdateKeySourcePatch()
 void FSubmixEffectMultibandCompressor::Initialize(FSubmixEffectMultibandCompressorSettings& Settings)
 {
 	const int32 NumBands = Settings.Bands.Num();
+
+	if(NumBands <= 0)
+	{
+		bInitialized = false;
+		DynamicsProcessors.Reset();
+		PrevCrossovers.Reset();
+		return;
+	}
+	
 	FrameSize = sizeof(float) * NumChannels;
 
 	PrevCrossovers.Reset(NumBands - 1);
@@ -322,7 +331,6 @@ void FSubmixEffectMultibandCompressor::Initialize(FSubmixEffectMultibandCompress
 
 	PrevNumBands = Settings.Bands.Num();
 	bPrevFourPole = Settings.bFourPole;
-
 	bInitialized = true;
 }
 
