[Engine/Plugins]

diff --git a/Engine/Plugins/Runtime/Synthesis/Source/Synthesis/Classes/SubmixEffects/SubmixEffectStereoDelay.h b/Engine/Plugins/Runtime/Synthesis/Source/Synthesis/Classes/SubmixEffects/SubmixEffectStereoDelay.h
--- a/Engine/Plugins/Runtime/Synthesis/Source/Synthesis/Classes/SubmixEffects/SubmixEffectStereoDelay.h
+++ b/Engine/Plugins/Runtime/Synthesis/Source/Synthesis/Classes/SubmixEffects/SubmixEffectStereoDelay.h
@@ -2,8 +2,6 @@
 
 #pragma once
 
-#include "CoreMinimal.h"
-#include "AudioEffect.h"
 #include "Sound/SoundEffectSubmix.h"
 #include "SourceEffects/SourceEffectStereoDelay.h"
 #include "SubmixEffectStereoDelay.generated.h"
@@ -100,3 +98,8 @@ public:
 	UPROPERTY(EditAnywhere, BlueprintReadWrite, Category = SubmixEffectPreset, Meta = (ShowOnlyInnerProperties))
 	FSubmixEffectStereoDelaySettings Settings;
 };
+
+#if UE_ENABLE_INCLUDE_ORDER_DEPRECATED_IN_5_2
+#include "AudioEffect.h"
+#include "CoreMinimal.h"
+#endif
diff --git a/Engine/Plugins/Runtime/Synthesis/Source/Synthesis/Classes/SubmixEffects/SubmixEffectTapDelay.h b/Engine/Plugins/Runtime/Synthesis/Source/Synthesis/Classes/SubmixEffects/SubmixEffectTapDelay.h
--- a/Engine/Plugins/Runtime/Synthesis/Source/Synthesis/Classes/SubmixEffects/SubmixEffectTapDelay.h
+++ b/Engine/Plugins/Runtime/Synthesis/Source/Synthesis/Classes/SubmixEffects/SubmixEffectTapDelay.h
@@ -2,10 +2,7 @@
 
 #pragma once
 
-#include "CoreMinimal.h"
-#include "AudioEffect.h"
 #include "DSP/Delay.h"
-#include "DSP/Dsp.h"
 #include "Sound/SoundEffectSubmix.h"
 #include "SubmixEffectTapDelay.generated.h"
 
@@ -222,3 +219,8 @@ public:
 
 	FSubmixEffectTapDelaySettings DynamicSettings;
 };
+
+#if UE_ENABLE_INCLUDE_ORDER_DEPRECATED_IN_5_2
+#include "AudioEffect.h"
+#include "CoreMinimal.h"
+#endif
diff --git a/Engine/Plugins/Runtime/Synthesis/Source/Synthesis/Classes/SynthComponents/EpicSynth1Component.h b/Engine/Plugins/Runtime/Synthesis/Source/Synthesis/Classes/SynthComponents/EpicSynth1Component.h
--- a/Engine/Plugins/Runtime/Synthesis/Source/Synthesis/Classes/SynthComponents/EpicSynth1Component.h
+++ b/Engine/Plugins/Runtime/Synthesis/Source/Synthesis/Classes/SynthComponents/EpicSynth1Component.h
@@ -2,19 +2,9 @@
 
 #pragma once
 
-#include "CoreMinimal.h"
-#include "Components/AudioComponent.h"
-#include "UObject/ObjectMacros.h"
-#include "UObject/Object.h"
-#include "UObject/Class.h"
-#include "Engine/EngineTypes.h"
-#include "Sound/SoundWaveProcedural.h"
 #include "Engine/DataTable.h"
 #include "Components/SynthComponent.h"
 #include "EpicSynth1.h"
-#include "EpicSynth1Types.h"
-#include "DSP/Osc.h"
-#include "Kismet/BlueprintFunctionLibrary.h"
 #include "EpicSynth1Component.generated.h"
 
 USTRUCT(BlueprintType)
@@ -665,3 +655,7 @@ public:
 protected:
 	Audio::FEpicSynth1 EpicSynth1;
 };
+
+#if UE_ENABLE_INCLUDE_ORDER_DEPRECATED_IN_5_2
+#include "CoreMinimal.h"
+#endif
diff --git a/Engine/Plugins/Runtime/Synthesis/Source/Synthesis/Classes/SynthComponents/SynthComponentGranulator.h b/Engine/Plugins/Runtime/Synthesis/Source/Synthesis/Classes/SynthComponents/SynthComponentGranulator.h
--- a/Engine/Plugins/Runtime/Synthesis/Source/Synthesis/Classes/SynthComponents/SynthComponentGranulator.h
+++ b/Engine/Plugins/Runtime/Synthesis/Source/Synthesis/Classes/SynthComponents/SynthComponentGranulator.h
@@ -2,11 +2,9 @@
 
 #pragma once
 
-#include "CoreMinimal.h"
 #include "Components/SynthComponent.h"
 #include "DSP/Granulator.h"
 #include "Sound/SampleBufferIO.h"
-#include "Sound/SoundWave.h"
 #include "SynthComponentGranulator.generated.h"
 
 
@@ -134,4 +132,8 @@ protected:
 	bool bRegistered;
 
 	bool bIsLoading;
-};
\ No newline at end of file
+};
+
+#if UE_ENABLE_INCLUDE_ORDER_DEPRECATED_IN_5_2
+#include "CoreMinimal.h"
+#endif
diff --git a/Engine/Plugins/Runtime/Synthesis/Source/Synthesis/Classes/SynthComponents/SynthComponentMonoWaveTable.h b/Engine/Plugins/Runtime/Synthesis/Source/Synthesis/Classes/SynthComponents/SynthComponentMonoWaveTable.h
--- a/Engine/Plugins/Runtime/Synthesis/Source/Synthesis/Classes/SynthComponents/SynthComponentMonoWaveTable.h
+++ b/Engine/Plugins/Runtime/Synthesis/Source/Synthesis/Classes/SynthComponents/SynthComponentMonoWaveTable.h
@@ -2,25 +2,13 @@
 
 #pragma once
 
-#include "CoreMinimal.h"
 #include "Components/SynthComponent.h"
-#include "DSP/Osc.h"
-#include "DSP/LFO.h"
-#include "Components/AudioComponent.h"
-#include "UObject/ObjectMacros.h"
-#include "UObject/Object.h"
-#include "UObject/Class.h"
-#include "Engine/EngineTypes.h"
-#include "Sound/SoundWaveProcedural.h"
-#include "Engine/DataTable.h"
 #include "Components/SynthComponent.h"
-#include "DSP/Osc.h"
-#include "Kismet/BlueprintFunctionLibrary.h"
 #include "MonoWaveTable.h"
-#include "EpicSynth1Types.h"
-#include "SynthesisModule.h"
 #include "SynthComponentMonoWaveTable.generated.h"
 
+enum class ESynthLFOType : uint8;
+
 
 class USynthComponentMonoWaveTable;
 
@@ -402,4 +390,11 @@ protected:
 	// underlying wavetable synth
 	Audio::FMonoWaveTable Synth;
 	Audio::DefaultWaveTableIndexType SampleRate;
-};
\ No newline at end of file
+};
+
+#if UE_ENABLE_INCLUDE_ORDER_DEPRECATED_IN_5_2
+#include "CoreMinimal.h"
+#include "Engine/DataTable.h"
+#include "EpicSynth1Types.h"
+#include "SynthesisModule.h"
+#endif
diff --git a/Engine/Plugins/Runtime/Synthesis/Source/Synthesis/Classes/SynthComponents/SynthComponentToneGenerator.h b/Engine/Plugins/Runtime/Synthesis/Source/Synthesis/Classes/SynthComponents/SynthComponentToneGenerator.h
--- a/Engine/Plugins/Runtime/Synthesis/Source/Synthesis/Classes/SynthComponents/SynthComponentToneGenerator.h
+++ b/Engine/Plugins/Runtime/Synthesis/Source/Synthesis/Classes/SynthComponents/SynthComponentToneGenerator.h
@@ -2,12 +2,9 @@
 
 #pragma once
 
-#include "CoreMinimal.h"
 #include "Components/SynthComponent.h"
-#include "Sound/SoundBase.h"
 #include "DSP/SinOsc.h"
 #include "DSP/AudioBufferDistanceAttenuation.h"
-#include "Sound/SoundGenerator.h"
 
 
 #include "SynthComponentToneGenerator.generated.h"
@@ -86,3 +83,7 @@ private:
 	Audio::FAudioBufferDistanceAttenuationSettings DistanceAttenuationSettings;
 	ISoundGeneratorPtr ToneGenerator;
 };
+
+#if UE_ENABLE_INCLUDE_ORDER_DEPRECATED_IN_5_2
+#include "CoreMinimal.h"
+#endif
diff --git a/Engine/Plugins/Runtime/Synthesis/Source/Synthesis/Classes/SynthComponents/SynthComponentWaveTable.h b/Engine/Plugins/Runtime/Synthesis/Source/Synthesis/Classes/SynthComponents/SynthComponentWaveTable.h
--- a/Engine/Plugins/Runtime/Synthesis/Source/Synthesis/Classes/SynthComponents/SynthComponentWaveTable.h
+++ b/Engine/Plugins/Runtime/Synthesis/Source/Synthesis/Classes/SynthComponents/SynthComponentWaveTable.h
@@ -2,12 +2,9 @@
 
 #pragma once
 
-#include "CoreMinimal.h"
 #include "Components/SynthComponent.h"
 #include "DSP/SampleBufferReader.h"
-#include "SampleBuffer.h"
 #include "Sound/SampleBufferIO.h"
-#include "Sound/SoundWave.h"
 #include "SynthComponentWaveTable.generated.h"
 
 UENUM(BlueprintType)
@@ -98,3 +95,7 @@ protected:
 
 	bool bIsLoaded;
 };
+
+#if UE_ENABLE_INCLUDE_ORDER_DEPRECATED_IN_5_2
+#include "CoreMinimal.h"
+#endif
diff --git a/Engine/Plugins/Runtime/Synthesis/Source/Synthesis/Private/EffectConvolutionReverb.cpp b/Engine/Plugins/Runtime/Synthesis/Source/Synthesis/Private/EffectConvolutionReverb.cpp
--- a/Engine/Plugins/Runtime/Synthesis/Source/Synthesis/Private/EffectConvolutionReverb.cpp
+++ b/Engine/Plugins/Runtime/Synthesis/Source/Synthesis/Private/EffectConvolutionReverb.cpp
@@ -2,6 +2,7 @@
 
 #include "EffectConvolutionReverb.h"
 #include "Async/AsyncWork.h"
+#include "DSP/AlignedBlockBuffer.h" // IWYU pragma: keep
 #include "SynthesisModule.h"
 
 #include UE_INLINE_GENERATED_CPP_BY_NAME(EffectConvolutionReverb)
diff --git a/Engine/Plugins/Runtime/Synthesis/Source/Synthesis/Private/SourceEffectConvolutionReverb.cpp b/Engine/Plugins/Runtime/Synthesis/Source/Synthesis/Private/SourceEffectConvolutionReverb.cpp
--- a/Engine/Plugins/Runtime/Synthesis/Source/Synthesis/Private/SourceEffectConvolutionReverb.cpp
+++ b/Engine/Plugins/Runtime/Synthesis/Source/Synthesis/Private/SourceEffectConvolutionReverb.cpp
@@ -2,6 +2,8 @@
 
 #include "SourceEffects/SourceEffectConvolutionReverb.h"
 
+#include "AudioMixerNullDevice.h"
+#include "DSP/AlignedBlockBuffer.h" // IWYU pragma: keep
 #include "DSP/FloatArrayMath.h"
 #include "EffectConvolutionReverb.h"
 
diff --git a/Engine/Plugins/Runtime/Synthesis/Source/Synthesis/Private/SubmixEffectConvolutionReverb.cpp b/Engine/Plugins/Runtime/Synthesis/Source/Synthesis/Private/SubmixEffectConvolutionReverb.cpp
--- a/Engine/Plugins/Runtime/Synthesis/Source/Synthesis/Private/SubmixEffectConvolutionReverb.cpp
+++ b/Engine/Plugins/Runtime/Synthesis/Source/Synthesis/Private/SubmixEffectConvolutionReverb.cpp
@@ -2,6 +2,8 @@
 
 #include "SubmixEffects/SubmixEffectConvolutionReverb.h"
 
+#include "AudioDefines.h"
+#include "DSP/AlignedBlockBuffer.h" // IWYU pragma: keep
 #include "DSP/FloatArrayMath.h"
 #include "EffectConvolutionReverb.h"
 
diff --git a/Engine/Plugins/Runtime/Synthesis/Source/Synthesis/Private/SubmixEffectFlexiverb.cpp b/Engine/Plugins/Runtime/Synthesis/Source/Synthesis/Private/SubmixEffectFlexiverb.cpp
--- a/Engine/Plugins/Runtime/Synthesis/Source/Synthesis/Private/SubmixEffectFlexiverb.cpp
+++ b/Engine/Plugins/Runtime/Synthesis/Source/Synthesis/Private/SubmixEffectFlexiverb.cpp
@@ -1,6 +1,8 @@
 // Copyright Epic Games, Inc. All Rights Reserved.
 
 #include "SubmixEffects/SubmixEffectFlexiverb.h"
+#include "DSP/AllPassFilter.h" // IWYU pragma: keep
+#include "DSP/OnePole.h" // IWYU pragma: keep
 
 #include UE_INLINE_GENERATED_CPP_BY_NAME(SubmixEffectFlexiverb)
 
diff --git a/Engine/Plugins/Runtime/Synthesis/Source/Synthesis/Public/ConvolutionReverb.h b/Engine/Plugins/Runtime/Synthesis/Source/Synthesis/Public/ConvolutionReverb.h
--- a/Engine/Plugins/Runtime/Synthesis/Source/Synthesis/Public/ConvolutionReverb.h
+++ b/Engine/Plugins/Runtime/Synthesis/Source/Synthesis/Public/ConvolutionReverb.h
@@ -2,13 +2,11 @@
 
 #pragma once
 
-#include "CoreMinimal.h"
-#include "DSP/AlignedBlockBuffer.h"
 #include "DSP/AudioChannelFormatConverter.h"
-#include "DSP/BufferVectorOperations.h"
 #include "DSP/ConvolutionAlgorithm.h"
 
-#include <type_traits>
+
+namespace Audio { class FAlignedBlockBuffer; }
 
 namespace Audio
 {
@@ -263,3 +261,9 @@ namespace Audio
 		bool bIsConvertingOutputChannelFormat;
 	};
 }
+
+#if UE_ENABLE_INCLUDE_ORDER_DEPRECATED_IN_5_2
+#include "CoreMinimal.h"
+#include "DSP/AlignedBlockBuffer.h"
+#include "DSP/BufferVectorOperations.h"
+#endif
diff --git a/Engine/Plugins/Runtime/Synthesis/Source/Synthesis/Public/EffectConvolutionReverb.h b/Engine/Plugins/Runtime/Synthesis/Source/Synthesis/Public/EffectConvolutionReverb.h
--- a/Engine/Plugins/Runtime/Synthesis/Source/Synthesis/Public/EffectConvolutionReverb.h
+++ b/Engine/Plugins/Runtime/Synthesis/Source/Synthesis/Public/EffectConvolutionReverb.h
@@ -1,13 +1,11 @@
 // Copyright Epic Games, Inc. All Rights Reserved.
 #pragma once
-#include "CoreMinimal.h"
-#include "Async/Async.h"
 #include "ConvolutionReverb.h"
-#include "DSP/ConvolutionAlgorithm.h"
 #include "DSP/Dsp.h"
-#include "SynthesisModule.h"
 #include "EffectConvolutionReverb.generated.h"
 
+struct FPropertyChangedEvent;
+
 /** Block size of convolution reverb algorithm. */
 UENUM()
 enum class ESubmixEffectConvolutionReverbBlockSize : uint8
@@ -153,3 +151,9 @@ namespace Audio
 		float SampleRate = 0.f;
 	};
 }
+
+#if UE_ENABLE_INCLUDE_ORDER_DEPRECATED_IN_5_2
+#include "Async/Async.h"
+#include "CoreMinimal.h"
+#include "SynthesisModule.h"
+#endif
diff --git a/Engine/Plugins/Runtime/Synthesis/Source/Synthesis/Public/EpicSynth1.h b/Engine/Plugins/Runtime/Synthesis/Source/Synthesis/Public/EpicSynth1.h
--- a/Engine/Plugins/Runtime/Synthesis/Source/Synthesis/Public/EpicSynth1.h
+++ b/Engine/Plugins/Runtime/Synthesis/Source/Synthesis/Public/EpicSynth1.h
@@ -2,16 +2,10 @@
 
 #pragma once
 
-#include "CoreMinimal.h"
 
-#include "DSP/Osc.h"
-#include "DSP/LFO.h"
-#include "DSP/Dsp.h"
 #include "DSP/Envelope.h"
 #include "DSP/Amp.h"
 #include "DSP/DelayStereo.h"
-#include "DSP/Filter.h"
-#include "DSP/ModulationMatrix.h"
 #include "DSP/Chorus.h"
 #include "EpicSynth1Types.h"
 
@@ -271,3 +265,7 @@ namespace Audio
 	};
 
 }
+
+#if UE_ENABLE_INCLUDE_ORDER_DEPRECATED_IN_5_2
+#include "CoreMinimal.h"
+#endif
diff --git a/Engine/Plugins/Runtime/Synthesis/Source/Synthesis/Public/Flexiverb.h b/Engine/Plugins/Runtime/Synthesis/Source/Synthesis/Public/Flexiverb.h
--- a/Engine/Plugins/Runtime/Synthesis/Source/Synthesis/Public/Flexiverb.h
+++ b/Engine/Plugins/Runtime/Synthesis/Source/Synthesis/Public/Flexiverb.h
@@ -1,11 +1,10 @@
 // Copyright Epic Games, Inc. All Rights Reserved.
 
 #pragma once
-#include "CoreMinimal.h"
-#include "DSP/Filter.h"
 #include "DSP/Delay.h"
-#include "DSP/AllPassFilter.h"
-#include "DSP/OnePole.h"
+
+namespace Audio { class FDelayAPF; }
+namespace Audio { class FOnePoleLPF; }
 
 namespace Audio
 {
@@ -93,4 +92,11 @@ namespace Audio
 		//Amount of delay lines in our FDN.
 		int32 ScatteringMatrixLength;
 	};
-}
\ No newline at end of file
+}
+
+#if UE_ENABLE_INCLUDE_ORDER_DEPRECATED_IN_5_2
+#include "CoreMinimal.h"
+#include "DSP/AllPassFilter.h"
+#include "DSP/Filter.h"
+#include "DSP/OnePole.h"
+#endif
diff --git a/Engine/Plugins/Runtime/Synthesis/Source/Synthesis/Public/SynthesisModule.h b/Engine/Plugins/Runtime/Synthesis/Source/Synthesis/Public/SynthesisModule.h
--- a/Engine/Plugins/Runtime/Synthesis/Source/Synthesis/Public/SynthesisModule.h
+++ b/Engine/Plugins/Runtime/Synthesis/Source/Synthesis/Public/SynthesisModule.h
@@ -2,8 +2,8 @@
 
 #pragma once
 
+#include "Logging/LogMacros.h"
 #include "Modules/ModuleInterface.h"
-#include "Modules/ModuleManager.h"
 
 DECLARE_LOG_CATEGORY_EXTERN(LogSynthesis, Log, All);
 
@@ -18,3 +18,7 @@ public:
 
 
 
+
+#if UE_ENABLE_INCLUDE_ORDER_DEPRECATED_IN_5_2
+#include "Modules/ModuleManager.h"
+#endif
diff --git a/Engine/Plugins/Runtime/Synthesis/Source/Synthesis/Public/UI/SSynth2DSlider.h b/Engine/Plugins/Runtime/Synthesis/Source/Synthesis/Public/UI/SSynth2DSlider.h
--- a/Engine/Plugins/Runtime/Synthesis/Source/Synthesis/Public/UI/SSynth2DSlider.h
+++ b/Engine/Plugins/Runtime/Synthesis/Source/Synthesis/Public/UI/SSynth2DSlider.h
@@ -2,14 +2,9 @@
 
 #pragma once
 
-#include "CoreMinimal.h"
-#include "Misc/Attribute.h"
-#include "Input/Reply.h"
 #include "Styling/SlateWidgetStyleAsset.h"
+#include "UI/SynthSlateStyle.h"
 #include "Widgets/DeclarativeSyntaxSupport.h"
-#include "Styling/SlateColor.h"
-#include "Styling/SlateTypes.h"
-#include "Styling/CoreStyle.h"
 #include "Framework/SlateDelegates.h"
 #include "Widgets/SLeafWidget.h"
 #include "Styling/SlateStyle.h"
@@ -205,3 +200,9 @@ private:
 	FOnFloatValueChanged OnValueChangedX;
 	FOnFloatValueChanged OnValueChangedY;
 };
+
+#if UE_ENABLE_INCLUDE_ORDER_DEPRECATED_IN_5_2
+#include "CoreMinimal.h"
+#include "Styling/CoreStyle.h"
+#include "Styling/SlateTypes.h"
+#endif
diff --git a/Engine/Plugins/Runtime/Synthesis/Source/Synthesis/Public/UI/SSynthKnob.h b/Engine/Plugins/Runtime/Synthesis/Source/Synthesis/Public/UI/SSynthKnob.h
--- a/Engine/Plugins/Runtime/Synthesis/Source/Synthesis/Public/UI/SSynthKnob.h
+++ b/Engine/Plugins/Runtime/Synthesis/Source/Synthesis/Public/UI/SSynthKnob.h
@@ -2,18 +2,13 @@
 
 #pragma once
 
-#include "CoreMinimal.h"
-#include "Misc/Attribute.h"
-#include "Input/Reply.h"
+#include "Styling/SlateStyle.h"
 #include "Styling/SlateWidgetStyleAsset.h"
+#include "UI/SynthSlateStyle.h"
 #include "Widgets/DeclarativeSyntaxSupport.h"
-#include "Styling/SlateColor.h"
-#include "Styling/SlateTypes.h"
-#include "Styling/CoreStyle.h"
 #include "Framework/SlateDelegates.h"
 #include "Widgets/SLeafWidget.h"
 #include "UI/SynthKnobStyle.h"
-#include "UI/SSynthTooltip.h"
 
 class FPaintArgs;
 class FSlateWindowElementList;
@@ -217,3 +212,10 @@ private:
 	// Holds a delegate that is executed when the slider's value changed.
 	FOnFloatValueChanged OnValueChanged;
 };
+
+#if UE_ENABLE_INCLUDE_ORDER_DEPRECATED_IN_5_2
+#include "CoreMinimal.h"
+#include "Styling/CoreStyle.h"
+#include "Styling/SlateTypes.h"
+#include "UI/SSynthTooltip.h"
+#endif
diff --git a/Engine/Plugins/Runtime/Synthesis/Source/Synthesis/Public/UI/SSynthTooltip.h b/Engine/Plugins/Runtime/Synthesis/Source/Synthesis/Public/UI/SSynthTooltip.h
--- a/Engine/Plugins/Runtime/Synthesis/Source/Synthesis/Public/UI/SSynthTooltip.h
+++ b/Engine/Plugins/Runtime/Synthesis/Source/Synthesis/Public/UI/SSynthTooltip.h
@@ -2,16 +2,9 @@
 
 #pragma once
 
-#include "CoreMinimal.h"
-#include "Misc/Attribute.h"
-#include "Styling/SlateWidgetStyleAsset.h"
-#include "Widgets/DeclarativeSyntaxSupport.h"
-#include "Styling/SlateColor.h"
-#include "Styling/SlateTypes.h"
-#include "Styling/CoreStyle.h"
-#include "Widgets/Text/STextBlock.h"
 #include "Widgets/SOverlay.h"
-#include "Widgets/SWindow.h"
+
+class STextBlock;
 
 // Special tooltip that doesn't follow the mouse position
 // Creates a simple window at the designated position.
@@ -36,3 +29,13 @@ private:
 	bool bIsVisible;
 
 };
+
+#if UE_ENABLE_INCLUDE_ORDER_DEPRECATED_IN_5_2
+#include "CoreMinimal.h"
+#include "Styling/CoreStyle.h"
+#include "Styling/SlateColor.h"
+#include "Styling/SlateTypes.h"
+#include "Styling/SlateWidgetStyleAsset.h"
+#include "Widgets/SWindow.h"
+#include "Widgets/Text/STextBlock.h"
+#endif
diff --git a/Engine/Plugins/Runtime/Synthesis/Source/Synthesis/Public/UI/Synth2DSlider.h b/Engine/Plugins/Runtime/Synthesis/Source/Synthesis/Public/UI/Synth2DSlider.h
--- a/Engine/Plugins/Runtime/Synthesis/Source/Synthesis/Public/UI/Synth2DSlider.h
+++ b/Engine/Plugins/Runtime/Synthesis/Source/Synthesis/Public/UI/Synth2DSlider.h
@@ -2,13 +2,7 @@
 
 #pragma once
 
-#include "CoreMinimal.h"
-#include "UI/SSynth2DSlider.h"
-#include "UI/SynthTypes.h"
 #include "UI/Synth2DSliderStyle.h"
-#include "UObject/ObjectMacros.h"
-#include "Styling/SlateTypes.h"
-#include "Widgets/SWidget.h"
 #include "Components/Widget.h"
 #include "Synth2DSlider.generated.h"
 
@@ -152,3 +146,11 @@ protected:
 	PROPERTY_BINDING_IMPLEMENTATION(float, ValueX);
 	PROPERTY_BINDING_IMPLEMENTATION(float, ValueY);
 };
+
+#if UE_ENABLE_INCLUDE_ORDER_DEPRECATED_IN_5_2
+#include "CoreMinimal.h"
+#include "Styling/SlateTypes.h"
+#include "UI/SSynth2DSlider.h"
+#include "UI/SynthTypes.h"
+#include "Widgets/SWidget.h"
+#endif
diff --git a/Engine/Plugins/Runtime/Synthesis/Source/Synthesis/Public/UI/Synth2DSliderStyle.h b/Engine/Plugins/Runtime/Synthesis/Source/Synthesis/Public/UI/Synth2DSliderStyle.h
--- a/Engine/Plugins/Runtime/Synthesis/Source/Synthesis/Public/UI/Synth2DSliderStyle.h
+++ b/Engine/Plugins/Runtime/Synthesis/Source/Synthesis/Public/UI/Synth2DSliderStyle.h
@@ -2,16 +2,10 @@
 
 #pragma once
 
-#include "CoreMinimal.h"
 
-#include "UI/SynthSlateStyle.h"
 
-#include "UObject/ObjectMacros.h"
-#include "Styling/SlateTypes.h"
-#include "Widgets/SWidget.h"
-#include "Components/Widget.h"
+#include "Styling/SlateBrush.h"
 #include "Styling/SlateWidgetStyle.h"
-#include "Brushes/SlateImageBrush.h"
 #include "Synth2DSliderStyle.generated.h"
 
 /**
@@ -57,3 +51,12 @@ struct SYNTHESIS_API FSynth2DSliderStyle : public FSlateWidgetStyle
 
 };
 
+
+#if UE_ENABLE_INCLUDE_ORDER_DEPRECATED_IN_5_2
+#include "Brushes/SlateImageBrush.h"
+#include "Components/Widget.h"
+#include "CoreMinimal.h"
+#include "Styling/SlateTypes.h"
+#include "UI/SynthSlateStyle.h"
+#include "Widgets/SWidget.h"
+#endif
diff --git a/Engine/Plugins/Runtime/Synthesis/Source/Synthesis/Public/UI/SynthKnob.h b/Engine/Plugins/Runtime/Synthesis/Source/Synthesis/Public/UI/SynthKnob.h
--- a/Engine/Plugins/Runtime/Synthesis/Source/Synthesis/Public/UI/SynthKnob.h
+++ b/Engine/Plugins/Runtime/Synthesis/Source/Synthesis/Public/UI/SynthKnob.h
@@ -2,20 +2,12 @@
 
 #pragma once
 
-#include "CoreMinimal.h"
-#include "UI/SSynthKnob.h"
-#include "UI/SSynthTooltip.h"
 #include "UI/SynthKnobStyle.h"
-#include "UI/SynthTypes.h"
-#include "UObject/ObjectMacros.h"
-#include "Styling/SlateTypes.h"
-#include "Widgets/SWidget.h"
 #include "Components/Slider.h"
-#include "Components/TextBlock.h"
-#include "Components/Widget.h"
-#include "Components/ContentWidget.h"
 #include "SynthKnob.generated.h"
 
+class SSynthKnob;
+
 /**
  * A simple widget that shows a sliding bar with a handle that allows you to control the value between 0..1.
  *
@@ -140,3 +132,13 @@ protected:
 protected:
 	PROPERTY_BINDING_IMPLEMENTATION(float, Value);
 };
+
+#if UE_ENABLE_INCLUDE_ORDER_DEPRECATED_IN_5_2
+#include "Components/ContentWidget.h"
+#include "Components/TextBlock.h"
+#include "CoreMinimal.h"
+#include "UI/SSynthKnob.h"
+#include "UI/SSynthTooltip.h"
+#include "UI/SynthTypes.h"
+#include "Widgets/SWidget.h"
+#endif
diff --git a/Engine/Plugins/Runtime/Synthesis/Source/Synthesis/Public/UI/SynthKnobStyle.h b/Engine/Plugins/Runtime/Synthesis/Source/Synthesis/Public/UI/SynthKnobStyle.h
--- a/Engine/Plugins/Runtime/Synthesis/Source/Synthesis/Public/UI/SynthKnobStyle.h
+++ b/Engine/Plugins/Runtime/Synthesis/Source/Synthesis/Public/UI/SynthKnobStyle.h
@@ -2,16 +2,10 @@
 
 #pragma once
 
-#include "CoreMinimal.h"
 
-#include "UI/SynthSlateStyle.h"
 
-#include "UObject/ObjectMacros.h"
-#include "Styling/SlateTypes.h"
-#include "Widgets/SWidget.h"
-#include "Components/Widget.h"
+#include "Styling/SlateBrush.h"
 #include "Styling/SlateWidgetStyle.h"
-#include "Brushes/SlateImageBrush.h"
 #include "SynthKnobStyle.generated.h"
 
 UENUM(BlueprintType)
@@ -80,3 +74,12 @@ struct SYNTHESIS_API FSynthKnobStyle : public FSlateWidgetStyle
 
 };
 
+
+#if UE_ENABLE_INCLUDE_ORDER_DEPRECATED_IN_5_2
+#include "Brushes/SlateImageBrush.h"
+#include "Components/Widget.h"
+#include "CoreMinimal.h"
+#include "Styling/SlateTypes.h"
+#include "UI/SynthSlateStyle.h"
+#include "Widgets/SWidget.h"
+#endif
diff --git a/Engine/Plugins/Runtime/Synthesis/Source/Synthesis/Public/UI/SynthSlateStyle.h b/Engine/Plugins/Runtime/Synthesis/Source/Synthesis/Public/UI/SynthSlateStyle.h
--- a/Engine/Plugins/Runtime/Synthesis/Source/Synthesis/Public/UI/SynthSlateStyle.h
+++ b/Engine/Plugins/Runtime/Synthesis/Source/Synthesis/Public/UI/SynthSlateStyle.h
@@ -2,16 +2,11 @@
 
 #pragma once
 
-#include "CoreMinimal.h"
-#include "UObject/ObjectMacros.h"
-#include "Styling/SlateTypes.h"
-#include "Widgets/SWidget.h"
-#include "Components/Widget.h"
 #include "Styling/SlateWidgetStyle.h"
-#include "Styling/SlateStyle.h"
-#include "Brushes/SlateImageBrush.h"
 #include "SynthSlateStyle.generated.h"
 
+struct FSlateDynamicImageBrush;
+
 UENUM(BlueprintType)
 enum class ESynthSlateSizeType : uint8
 {
@@ -94,3 +89,12 @@ protected:
 	static ISynthSlateResources* SynthSlateResources;
 };
 
+
+#if UE_ENABLE_INCLUDE_ORDER_DEPRECATED_IN_5_2
+#include "Brushes/SlateImageBrush.h"
+#include "Components/Widget.h"
+#include "CoreMinimal.h"
+#include "Styling/SlateStyle.h"
+#include "Styling/SlateTypes.h"
+#include "Widgets/SWidget.h"
+#endif
diff --git a/Engine/Plugins/Runtime/Synthesis/Source/Synthesis/Public/UI/SynthTypes.h b/Engine/Plugins/Runtime/Synthesis/Source/Synthesis/Public/UI/SynthTypes.h
--- a/Engine/Plugins/Runtime/Synthesis/Source/Synthesis/Public/UI/SynthTypes.h
+++ b/Engine/Plugins/Runtime/Synthesis/Source/Synthesis/Public/UI/SynthTypes.h
@@ -2,6 +2,9 @@
 
 #pragma once
 
+
+
+#if UE_ENABLE_INCLUDE_ORDER_DEPRECATED_IN_5_2
 #include "CoreMinimal.h"
 #include "UObject/ObjectMacros.h"
-
+#endif
diff --git a/Engine/Plugins/Runtime/Synthesis/Source/SynthesisEditor/Classes/AudioImpulseResponseAsset.h b/Engine/Plugins/Runtime/Synthesis/Source/SynthesisEditor/Classes/AudioImpulseResponseAsset.h
--- a/Engine/Plugins/Runtime/Synthesis/Source/SynthesisEditor/Classes/AudioImpulseResponseAsset.h
+++ b/Engine/Plugins/Runtime/Synthesis/Source/SynthesisEditor/Classes/AudioImpulseResponseAsset.h
@@ -2,11 +2,8 @@
 
 #pragma once
 
-#include "CoreMinimal.h"
 #include "Factories/Factory.h"
-#include "AssetTypeCategories.h"
 #include "AssetTypeActions_Base.h"
-#include "Sound/SoundWave.h"
 #include "Sound/SampleBufferIO.h"
 #include "AudioImpulseResponseAsset.generated.h"
 
@@ -42,3 +39,8 @@ class UAudioImpulseResponseFactory : public UFactory
 private:
 	Audio::FSoundWavePCMLoader Loader;
 };
+
+#if UE_ENABLE_INCLUDE_ORDER_DEPRECATED_IN_5_2
+#include "CoreMinimal.h"
+#include "Sound/SoundWave.h"
+#endif
diff --git a/Engine/Plugins/Runtime/Synthesis/Source/SynthesisEditor/Classes/EpicSynth1PresetBank.h b/Engine/Plugins/Runtime/Synthesis/Source/SynthesisEditor/Classes/EpicSynth1PresetBank.h
--- a/Engine/Plugins/Runtime/Synthesis/Source/SynthesisEditor/Classes/EpicSynth1PresetBank.h
+++ b/Engine/Plugins/Runtime/Synthesis/Source/SynthesisEditor/Classes/EpicSynth1PresetBank.h
@@ -2,9 +2,6 @@
 
 #pragma once
 
-#include "CoreMinimal.h"
-#include "Stats/Stats.h"
-#include "AssetToolsModule.h"
 #include "Factories/Factory.h"
 #include "AssetTypeActions_Base.h"
 #include "EpicSynth1PresetBank.generated.h"
@@ -28,3 +25,7 @@ class UModularSynthPresetBankFactory : public UFactory
 	virtual UObject* FactoryCreateNew(UClass* Class, UObject* InParent, FName Name, EObjectFlags Flags, UObject* Context, FFeedbackContext* Warn) override;
 	//~ Begin UFactory Interface	
 };
+
+#if UE_ENABLE_INCLUDE_ORDER_DEPRECATED_IN_5_2
+#include "CoreMinimal.h"
+#endif
diff --git a/Engine/Plugins/Runtime/Synthesis/Source/SynthesisEditor/Classes/MonoWaveTablePresetBank.h b/Engine/Plugins/Runtime/Synthesis/Source/SynthesisEditor/Classes/MonoWaveTablePresetBank.h
--- a/Engine/Plugins/Runtime/Synthesis/Source/SynthesisEditor/Classes/MonoWaveTablePresetBank.h
+++ b/Engine/Plugins/Runtime/Synthesis/Source/SynthesisEditor/Classes/MonoWaveTablePresetBank.h
@@ -2,9 +2,6 @@
 
 #pragma once
 
-#include "CoreMinimal.h"
-#include "Stats/Stats.h"
-#include "AssetToolsModule.h"
 #include "Factories/Factory.h"
 #include "AssetTypeActions_Base.h"
 #include "MonoWaveTablePresetBank.generated.h"
@@ -31,3 +28,7 @@ class UMonoWaveTableSynthPresetFactory : public UFactory
 
 
 
+
+#if UE_ENABLE_INCLUDE_ORDER_DEPRECATED_IN_5_2
+#include "CoreMinimal.h"
+#endif
diff --git a/Engine/Plugins/Runtime/Synthesis/Source/SynthesisEditor/Classes/SynthesisEditorSettings.h b/Engine/Plugins/Runtime/Synthesis/Source/SynthesisEditor/Classes/SynthesisEditorSettings.h
--- a/Engine/Plugins/Runtime/Synthesis/Source/SynthesisEditor/Classes/SynthesisEditorSettings.h
+++ b/Engine/Plugins/Runtime/Synthesis/Source/SynthesisEditor/Classes/SynthesisEditorSettings.h
@@ -2,9 +2,8 @@
 
 #pragma once
 
-#include "AudioEditorSettings.h"
-#include "WidgetBlueprint.h"
 
+#include "Engine/DeveloperSettings.h"
 #include "SynthesisEditorSettings.generated.h"
 
 
@@ -13,3 +12,8 @@ class USynthesisEditorSettings : public UDeveloperSettings
 {
 	GENERATED_BODY()
 };
+
+#if UE_ENABLE_INCLUDE_ORDER_DEPRECATED_IN_5_2
+#include "AudioEditorSettings.h"
+#include "WidgetBlueprint.h"
+#endif
diff --git a/Engine/Plugins/Runtime/Synthesis/Source/SynthesisEditor/Public/SynthesisEditorModule.h b/Engine/Plugins/Runtime/Synthesis/Source/SynthesisEditor/Public/SynthesisEditorModule.h
--- a/Engine/Plugins/Runtime/Synthesis/Source/SynthesisEditor/Public/SynthesisEditorModule.h
+++ b/Engine/Plugins/Runtime/Synthesis/Source/SynthesisEditor/Public/SynthesisEditorModule.h
@@ -2,8 +2,8 @@
 
 #pragma once
 
+#include "Logging/LogMacros.h"
 #include "Modules/ModuleInterface.h"
-#include "Modules/ModuleManager.h"
 
 
 DECLARE_LOG_CATEGORY_EXTERN(LogSynthesisEditor, Log, All);
@@ -20,3 +20,7 @@ public:
 private:
 	void RegisterMenus();
 };
+
+#if UE_ENABLE_INCLUDE_ORDER_DEPRECATED_IN_5_2
+#include "Modules/ModuleManager.h"
+#endif
diff --git a/Engine/Plugins/Runtime/WaveTable/Source/WaveTable/Public/WaveTable.h b/Engine/Plugins/Runtime/WaveTable/Source/WaveTable/Public/WaveTable.h
--- a/Engine/Plugins/Runtime/WaveTable/Source/WaveTable/Public/WaveTable.h
+++ b/Engine/Plugins/Runtime/WaveTable/Source/WaveTable/Public/WaveTable.h
@@ -1,9 +1,7 @@
 // Copyright Epic Games, Inc. All Rights Reserved.
 #pragma once
 
-#include "Containers/Array.h"
 #include "Containers/ArrayView.h"
-#include "HAL/Platform.h"
 
 
 namespace WaveTable
@@ -102,4 +100,4 @@ namespace WaveTable
 			FMemory::Memset(Samples.GetData(), 0, sizeof(float) * Samples.Num());
 		}
 	};
-} // namespace WaveTable
\ No newline at end of file
+} // namespace WaveTable
diff --git a/Engine/Plugins/Runtime/WaveTable/Source/WaveTable/Public/WaveTableBank.h b/Engine/Plugins/Runtime/WaveTable/Source/WaveTable/Public/WaveTableBank.h
--- a/Engine/Plugins/Runtime/WaveTable/Source/WaveTable/Public/WaveTableBank.h
+++ b/Engine/Plugins/Runtime/WaveTable/Source/WaveTable/Public/WaveTableBank.h
@@ -1,16 +1,14 @@
 // Copyright Epic Games, Inc. All Rights Reserved.
 #pragma once
 
-#include "Algo/Transform.h"
 #include "IAudioProxyInitializer.h"
-#include "UObject/Object.h"
-#include "UObject/ObjectMacros.h"
-#include "UObject/ObjectSaveContext.h"
 #include "WaveTable.h"
 #include "WaveTableTransform.h"
 
 #include "WaveTableBank.generated.h"
 
+struct FPropertyChangedEvent;
+
 
 USTRUCT()
 struct WAVETABLE_API FWaveTableBankEntry
@@ -92,3 +90,7 @@ protected:
 	TArray<WaveTable::FWaveTable> WaveTables;
 };
 using FWaveTableBankAssetProxyPtr = TSharedPtr<FWaveTableBankAssetProxy, ESPMode::ThreadSafe>;
+
+#if UE_ENABLE_INCLUDE_ORDER_DEPRECATED_IN_5_2
+#include "UObject/ObjectSaveContext.h"
+#endif
diff --git a/Engine/Plugins/Runtime/WaveTable/Source/WaveTable/Public/WaveTableImporter.h b/Engine/Plugins/Runtime/WaveTable/Source/WaveTable/Public/WaveTableImporter.h
--- a/Engine/Plugins/Runtime/WaveTable/Source/WaveTable/Public/WaveTableImporter.h
+++ b/Engine/Plugins/Runtime/WaveTable/Source/WaveTable/Public/WaveTableImporter.h
@@ -1,10 +1,10 @@
 // Copyright Epic Games, Inc. All Rights Reserved.
 #pragma once
 
-#include "Containers/Array.h"
-#include "WaveTableSettings.h"
 #include "WaveTableSampler.h"
 
+struct FWaveTableSettings;
+
 
 namespace WaveTable
 {
@@ -23,3 +23,7 @@ namespace WaveTable
 		FWaveTableSampler Sampler;
 	};
 } // namespace WaveTable
+
+#if UE_ENABLE_INCLUDE_ORDER_DEPRECATED_IN_5_2
+#include "WaveTableSettings.h"
+#endif
diff --git a/Engine/Plugins/Runtime/WaveTable/Source/WaveTable/Public/WaveTableSampler.h b/Engine/Plugins/Runtime/WaveTable/Source/WaveTable/Public/WaveTableSampler.h
--- a/Engine/Plugins/Runtime/WaveTable/Source/WaveTable/Public/WaveTableSampler.h
+++ b/Engine/Plugins/Runtime/WaveTable/Source/WaveTable/Public/WaveTableSampler.h
@@ -1,10 +1,7 @@
 // Copyright Epic Games, Inc. All Rights Reserved.
 #pragma once
 
-#include "Containers/Array.h"
-#include "DSP/Filter.h"
-#include "HAL/Platform.h"
-#include "WaveTableSettings.h"
+#include "Containers/ArrayView.h"
 
 
 namespace WaveTable
@@ -90,4 +87,9 @@ namespace WaveTable
 		TArray<float> PhaseModScratch;
 		FSettings Settings;
 	};
-} // namespace WaveTable
\ No newline at end of file
+} // namespace WaveTable
+
+#if UE_ENABLE_INCLUDE_ORDER_DEPRECATED_IN_5_2
+#include "DSP/Filter.h"
+#include "WaveTableSettings.h"
+#endif
diff --git a/Engine/Plugins/Runtime/WaveTable/Source/WaveTable/Public/WaveTableSettings.h b/Engine/Plugins/Runtime/WaveTable/Source/WaveTable/Public/WaveTableSettings.h
--- a/Engine/Plugins/Runtime/WaveTable/Source/WaveTable/Public/WaveTableSettings.h
+++ b/Engine/Plugins/Runtime/WaveTable/Source/WaveTable/Public/WaveTableSettings.h
@@ -1,13 +1,6 @@
 // Copyright Epic Games, Inc. All Rights Reserved.
 #pragma once
 
-#include "Containers/Array.h"
-#include "Containers/ArrayView.h"
-#include "GenericPlatform/GenericPlatformMath.h"
-#include "HAL/Platform.h"
-#include "UObject/SoftObjectPath.h"
-#include "UObject/Object.h"
-#include "UObject/ObjectMacros.h"
 
 #include "WaveTableSettings.generated.h"
 
@@ -134,4 +127,4 @@ namespace WaveTable
 	// Returns WaveTable size given WaveTableResolution and optional Curve or PCM size. If resolution is none, uses default value associated
 	// with provided curve. If set to file, will use either the largest static resolution (Res_Max, or 4096) or InMaxPCMSize, whichever is larger.
 	WAVETABLE_API int32 GetWaveTableSize(EWaveTableResolution InWaveTableResolution, EWaveTableCurve InCurve, int32 InMaxPCMSize = INDEX_NONE);
-} // namespace WaveTable
\ No newline at end of file
+} // namespace WaveTable
diff --git a/Engine/Plugins/Runtime/WaveTable/Source/WaveTable/Public/WaveTableTransform.h b/Engine/Plugins/Runtime/WaveTable/Source/WaveTable/Public/WaveTableTransform.h
--- a/Engine/Plugins/Runtime/WaveTable/Source/WaveTable/Public/WaveTableTransform.h
+++ b/Engine/Plugins/Runtime/WaveTable/Source/WaveTable/Public/WaveTableTransform.h
@@ -1,14 +1,13 @@
 // Copyright Epic Games, Inc. All Rights Reserved.
 #pragma once
 
-#include "Containers/Array.h"
-#include "Containers/ArrayView.h"
-#include "Curves/CurveFloat.h"
 #include "Curves/RichCurve.h"
 #include "WaveTableSettings.h"
 
 #include "WaveTableTransform.generated.h"
 
+class UCurveFloat;
+
 // use a value outside the [-1,1] range to detect uninitialized values
 namespace WaveTable
 {
@@ -80,3 +79,7 @@ private:
 	/** Clamps & applies transform to provided values as unipolar signal */
 	void SampleCurveUnipolar(TArrayView<float> InOutValues, float& OutFinalValue) const;
 };
+
+#if UE_ENABLE_INCLUDE_ORDER_DEPRECATED_IN_5_2
+#include "Curves/CurveFloat.h"
+#endif
diff --git a/Engine/Plugins/Runtime/WaveTable/Source/WaveTableEditor/Public/WaveTableBankEditor.h b/Engine/Plugins/Runtime/WaveTable/Source/WaveTableEditor/Public/WaveTableBankEditor.h
--- a/Engine/Plugins/Runtime/WaveTable/Source/WaveTableEditor/Public/WaveTableBankEditor.h
+++ b/Engine/Plugins/Runtime/WaveTable/Source/WaveTableEditor/Public/WaveTableBankEditor.h
@@ -3,16 +3,15 @@
 #pragma once
 
 #include "CurveEditorTypes.h"
-#include "Curves/CurveOwnerInterface.h"
 #include "EditorUndoClient.h"
-#include "Framework/Docking/TabManager.h"
-#include "Math/Color.h"
 #include "Misc/NotifyHook.h"
-#include "WaveTableCurveEditorViewStacked.h"
 #include "Toolkits/AssetEditorToolkit.h"
-#include "Toolkits/IToolkitHost.h"
 #include "WaveTableSettings.h"
-#include "Widgets/SWidget.h"
+
+class FSpawnTabArgs;
+enum class EWaveTableCurveSource : uint8;
+namespace WaveTable::Editor { class FWaveTableCurveModel; }
+struct FRichCurve;
 
 
 // Forward Declarations
@@ -147,4 +146,10 @@ namespace WaveTable
 			virtual TUniquePtr<FWaveTableCurveModel> ConstructCurveModel(FRichCurve& InRichCurve, UObject* InParentObject, EWaveTableCurveSource InSource) override;
 		};
 	} // namespace Editor
-} // namespace WaveTable
\ No newline at end of file
+} // namespace WaveTable
+
+#if UE_ENABLE_INCLUDE_ORDER_DEPRECATED_IN_5_2
+#include "Curves/CurveOwnerInterface.h"
+#include "Toolkits/IToolkitHost.h"
+#include "WaveTableCurveEditorViewStacked.h"
+#endif
diff --git a/Engine/Plugins/Runtime/WaveTable/Source/WaveTableEditor/Public/WaveTableCurveEditorViewStacked.h b/Engine/Plugins/Runtime/WaveTable/Source/WaveTableEditor/Public/WaveTableCurveEditorViewStacked.h
--- a/Engine/Plugins/Runtime/WaveTable/Source/WaveTableEditor/Public/WaveTableCurveEditorViewStacked.h
+++ b/Engine/Plugins/Runtime/WaveTable/Source/WaveTableEditor/Public/WaveTableCurveEditorViewStacked.h
@@ -1,16 +1,10 @@
 // Copyright Epic Games, Inc. All Rights Reserved.
 #pragma once
 
-#include "CurveModel.h"
-#include "Curves/RichCurve.h"
-#include "Internationalization/Text.h"
 #include "RichCurveEditorModel.h"
-#include "SGraphActionMenu.h"
-#include "UObject/Object.h"
-#include "UObject/WeakObjectPtr.h"
-#include "UObject/WeakObjectPtrTemplates.h"
 #include "Views/SCurveEditorViewStacked.h"
-#include "Views/SInteractiveCurveEditorView.h"
+
+enum class ECurveEditorViewID : uint64;
 
 #include "WaveTableCurveEditorViewStacked.generated.h"
 
@@ -146,3 +140,8 @@ namespace WaveTable
 		};
 	} // namespace Editor
 } // namespace WaveTable
+
+#if UE_ENABLE_INCLUDE_ORDER_DEPRECATED_IN_5_2
+#include "Curves/RichCurve.h"
+#include "SGraphActionMenu.h"
+#endif
diff --git a/Engine/Plugins/Runtime/WaveTable/Source/WaveTableEditor/Public/WaveTableFileUtilities.h b/Engine/Plugins/Runtime/WaveTable/Source/WaveTableEditor/Public/WaveTableFileUtilities.h
--- a/Engine/Plugins/Runtime/WaveTable/Source/WaveTableEditor/Public/WaveTableFileUtilities.h
+++ b/Engine/Plugins/Runtime/WaveTable/Source/WaveTableEditor/Public/WaveTableFileUtilities.h
@@ -2,7 +2,6 @@
 #pragma once
 
 #include "Containers/Array.h"
-#include "Containers/UnrealString.h"
 
 
 namespace WaveTable
@@ -15,3 +14,7 @@ namespace WaveTable
 		} // namespace FileUtilities
 	} // namespace Editor
 } // namespace WaveTable
+
+#if UE_ENABLE_INCLUDE_ORDER_DEPRECATED_IN_5_2
+#include "Containers/UnrealString.h"
+#endif
diff --git a/Engine/Plugins/Runtime/WaveTable/Source/WaveTableEditor/Public/WaveTableTransformLayout.h b/Engine/Plugins/Runtime/WaveTable/Source/WaveTableEditor/Public/WaveTableTransformLayout.h
--- a/Engine/Plugins/Runtime/WaveTable/Source/WaveTableEditor/Public/WaveTableTransformLayout.h
+++ b/Engine/Plugins/Runtime/WaveTable/Source/WaveTableEditor/Public/WaveTableTransformLayout.h
@@ -1,16 +1,10 @@
 // Copyright Epic Games, Inc. All Rights Reserved.
 #pragma once
 
-#include "Containers/Map.h"
-#include "Curves/SimpleCurve.h"
 #include "IPropertyTypeCustomization.h"
-#include "Layout/Visibility.h"
-#include "Misc/Attribute.h"
-#include "PropertyCustomizationHelpers.h"
-#include "PropertyHandle.h"
-#include "SCurveEditor.h"
-#include "WaveTableSettings.h"
-#include "WaveTableTransform.h"
+
+enum class EWaveTableCurve : uint8;
+struct FWaveTableTransform;
 
 
 namespace WaveTable
@@ -59,3 +53,14 @@ namespace WaveTable
 		};
 	} // namespace Editor
 } // namespace WaveTable
+
+#if UE_ENABLE_INCLUDE_ORDER_DEPRECATED_IN_5_2
+#include "Curves/SimpleCurve.h"
+#include "Layout/Visibility.h"
+#include "Misc/Attribute.h"
+#include "PropertyCustomizationHelpers.h"
+#include "PropertyHandle.h"
+#include "SCurveEditor.h"
+#include "WaveTableSettings.h"
+#include "WaveTableTransform.h"
+#endif
