Change Chorus Wet/Dry Modulation Parameter to volume (for parity with submix modulators)

diff --git a/Engine/Plugins/Runtime/Synthesis/Source/Synthesis/Classes/SourceEffects/SourceEffectChorus.h b/Engine/Plugins/Runtime/Synthesis/Source/Synthesis/Classes/SourceEffects/SourceEffectChorus.h
--- a/Engine/Plugins/Runtime/Synthesis/Source/Synthesis/Classes/SourceEffects/SourceEffectChorus.h
+++ b/Engine/Plugins/Runtime/Synthesis/Source/Synthesis/Classes/SourceEffects/SourceEffectChorus.h
@@ -74,11 +74,11 @@ struct SYNTHESIS_API FSourceEffectChorusSettings
 	FSoundModulationDestinationSettings FeedbackModulation;
 
 	// The wet level of the chorus effect
-	UPROPERTY(EditAnywhere, BlueprintReadWrite, Category = "SourceEffect|Preset", meta = (DisplayName = "Wet Level", AudioParamClass = "SoundModulationParameter", ClampMin = "0.0", ClampMax = "1.0", UIMin = "0.0", UIMax = "1.0"))
+	UPROPERTY(EditAnywhere, BlueprintReadWrite, Category = "SourceEffect|Preset", meta = (DisplayName = "Wet Level", AudioParamClass = "SoundModulationParameterVolume", ClampMin = "0.0", ClampMax = "1.0", UIMin = "0.0", UIMax = "1.0"))
 	FSoundModulationDestinationSettings WetModulation;
 
 	// The dry level of the chorus effect
-	UPROPERTY(EditAnywhere, BlueprintReadWrite, Category = "SourceEffect|Preset", meta = (DisplayName = "Dry Level", AudioParamClass = "SoundModulationParameter", ClampMin = "0.0", ClampMax = "1.0", UIMin = "0.0", UIMax = "1.0"))
+	UPROPERTY(EditAnywhere, BlueprintReadWrite, Category = "SourceEffect|Preset", meta = (DisplayName = "Dry Level", AudioParamClass = "SoundModulationParameterVolume", ClampMin = "0.0", ClampMax = "1.0", UIMin = "0.0", UIMax = "1.0"))
 	FSoundModulationDestinationSettings DryModulation;
 
 	// The spread of the effect (larger means greater difference between left and right delay lines)
