Reverted const correctness for sound effect submix change

diff --git a/Engine/Plugins/Runtime/Synthesis/Source/Synthesis/Private/SubmixEffectFilter.cpp b/Engine/Plugins/Runtime/Synthesis/Source/Synthesis/Private/SubmixEffectFilter.cpp
--- a/Engine/Plugins/Runtime/Synthesis/Source/Synthesis/Private/SubmixEffectFilter.cpp
+++ b/Engine/Plugins/Runtime/Synthesis/Source/Synthesis/Private/SubmixEffectFilter.cpp
@@ -52,7 +52,7 @@ void FSubmixEffectFilter::OnProcessAudio(const FSoundEffectSubmixInputData& InDa
 		InitFilter();
 	}
 
-	const float* InAudioBuffer = InData.AudioBuffer->GetData();
+	float* InAudioBuffer = InData.AudioBuffer->GetData();
 	float* OutAudioBuffer = OutData.AudioBuffer->GetData();
 	const int32 NumSamples = InData.AudioBuffer->Num();
 
diff --git a/Engine/Plugins/Runtime/Synthesis/Source/Synthesis/Private/SubmixEffectStereoDelay.cpp b/Engine/Plugins/Runtime/Synthesis/Source/Synthesis/Private/SubmixEffectStereoDelay.cpp
--- a/Engine/Plugins/Runtime/Synthesis/Source/Synthesis/Private/SubmixEffectStereoDelay.cpp
+++ b/Engine/Plugins/Runtime/Synthesis/Source/Synthesis/Private/SubmixEffectStereoDelay.cpp
@@ -30,7 +30,7 @@ void FSubmixEffectStereoDelay::OnProcessAudio(const FSoundEffectSubmixInputData&
 		ScratchStereoBuffer.Reset();
 		ScratchStereoBuffer.AddUninitialized(InData.NumFrames * 2);
 
-		const float* InputBufferPtr = InData.AudioBuffer->GetData();
+		float* InputBufferPtr = InData.AudioBuffer->GetData();
 
 		// Get the stereo data from the input buffer
 		for (int32 Frame = 0; Frame < InData.NumFrames; ++Frame)
@@ -60,7 +60,7 @@ void FSubmixEffectStereoDelay::OnProcessAudio(const FSoundEffectSubmixInputData&
 	}
 	else
 	{
-		const float* InputBufferPtr = InData.AudioBuffer->GetData();
+		float* InputBufferPtr = InData.AudioBuffer->GetData();
 		float* OutputBufferPtr = OutData.AudioBuffer->GetData();
 		int32 NumSamples = InData.NumChannels * InData.NumFrames;
 		DelayStereo.ProcessAudio(InputBufferPtr, NumSamples, OutputBufferPtr);
diff --git a/Engine/Plugins/Runtime/Synthesis/Source/Synthesis/Private/SubmixEffectStereoToQuad.cpp b/Engine/Plugins/Runtime/Synthesis/Source/Synthesis/Private/SubmixEffectStereoToQuad.cpp
--- a/Engine/Plugins/Runtime/Synthesis/Source/Synthesis/Private/SubmixEffectStereoToQuad.cpp
+++ b/Engine/Plugins/Runtime/Synthesis/Source/Synthesis/Private/SubmixEffectStereoToQuad.cpp
@@ -25,7 +25,7 @@ void FSubmixEffectStereoToQuad::OnProcessAudio(const FSoundEffectSubmixInputData
 
 	if (InData.NumChannels > 2)
 	{
-		const float* InputBufferPtr = InData.AudioBuffer->GetData();
+		float* InputBufferPtr = InData.AudioBuffer->GetData();
 		float* OutputBufferPtr = OutData.AudioBuffer->GetData();
 
 		// Channel offset in SMPTE channel order for side left and side right
diff --git a/Engine/Source/Runtime/AudioMixer/Private/Effects/AudioMixerSubmixEffectEQ.cpp b/Engine/Source/Runtime/AudioMixer/Private/Effects/AudioMixerSubmixEffectEQ.cpp
--- a/Engine/Source/Runtime/AudioMixer/Private/Effects/AudioMixerSubmixEffectEQ.cpp
+++ b/Engine/Source/Runtime/AudioMixer/Private/Effects/AudioMixerSubmixEffectEQ.cpp
@@ -102,7 +102,7 @@ void FSubmixEffectSubmixEQ::OnProcessAudio(const FSoundEffectSubmixInputData& In
 	// Update parameters that may have been set from game thread
 	UpdateParameters(InData.NumChannels);
 
-	const Audio::FAlignedFloatBuffer& InAudioBuffer = *InData.AudioBuffer;
+	Audio::FAlignedFloatBuffer& InAudioBuffer = *InData.AudioBuffer;
 	Audio::FAlignedFloatBuffer& OutAudioBuffer = *OutData.AudioBuffer;
 
 	if (bEQSettingsSet && RenderThreadEQSettings.EQBands.Num() > 0)
diff --git a/Engine/Source/Runtime/Engine/Classes/Sound/SoundEffectSubmix.h b/Engine/Source/Runtime/Engine/Classes/Sound/SoundEffectSubmix.h
--- a/Engine/Source/Runtime/Engine/Classes/Sound/SoundEffectSubmix.h
+++ b/Engine/Source/Runtime/Engine/Classes/Sound/SoundEffectSubmix.h
@@ -58,7 +58,7 @@ struct FSoundEffectSubmixInputData
 	const TArray<FTransform>* ListenerTransforms;
 
 	/** The raw input audio buffer. Size is NumFrames * NumChannels */
-	const Audio::FAlignedFloatBuffer* AudioBuffer;
+	Audio::FAlignedFloatBuffer* AudioBuffer;
 
 	/** Sample accurate audio clock. */
 	double AudioClock;
