Adding tooltips to source effecft dynamics processor.

Tooltips pulled from submix dynamics processor

[RN] Added tooltips, categories, and a bypass option to source effect dynamics processor

diff --git a/Engine/Plugins/Runtime/Synthesis/Source/Synthesis/Classes/SourceEffects/SourceEffectDynamicsProcessor.h b/Engine/Plugins/Runtime/Synthesis/Source/Synthesis/Classes/SourceEffects/SourceEffectDynamicsProcessor.h
--- a/Engine/Plugins/Runtime/Synthesis/Source/Synthesis/Classes/SourceEffects/SourceEffectDynamicsProcessor.h
+++ b/Engine/Plugins/Runtime/Synthesis/Source/Synthesis/Classes/SourceEffects/SourceEffectDynamicsProcessor.h
@@ -31,55 +31,62 @@ struct SYNTHESIS_API FSourceEffectDynamicsProcessorSettings
 {
 	GENERATED_USTRUCT_BODY()
 
-	UPROPERTY(EditAnywhere, BlueprintReadWrite, Category = "SourceEffect|Preset")
-	ESourceEffectDynamicsProcessorType DynamicsProcessorType;
+	// Type of processor to apply
+	UPROPERTY(EditAnywhere, BlueprintReadWrite, Category = General, meta = (DisplayName = "Type"))
+	ESourceEffectDynamicsProcessorType DynamicsProcessorType = ESourceEffectDynamicsProcessorType::Compressor;
 
-	UPROPERTY(EditAnywhere, BlueprintReadWrite, Category = "SourceEffect|Preset")
-	ESourceEffectDynamicsPeakMode PeakMode;
+	// Mode of peak detection used on input key signal
+	UPROPERTY(EditAnywhere, BlueprintReadWrite, Category = Dynamics, meta = (EditCondition = "!bBypass"))
+	ESourceEffectDynamicsPeakMode PeakMode = ESourceEffectDynamicsPeakMode::RootMeanSquared;
 
-	UPROPERTY(EditAnywhere, BlueprintReadWrite, Category = "SourceEffect|Preset",  meta = (ClampMin = "0.0", ClampMax = "50.0", UIMin = "0.0", UIMax = "50.0"))
-	float LookAheadMsec;
+	// The amount of time to look ahead of the current audio (Allows for transients to be included in dynamics processing)
+	UPROPERTY(EditAnywhere, BlueprintReadWrite, Category = Response,  meta = (DisplayName = "Look Ahead (ms)", ClampMin = "0.0", ClampMax = "50.0", UIMin = "0.0", UIMax = "50.0", EditCondition = "!bBypass"))
+	float LookAheadMsec = 3.0f;
 
-	UPROPERTY(EditAnywhere, BlueprintReadWrite, Category = "SourceEffect|Preset", meta = (ClampMin = "1.0", ClampMax = "300.0", UIMin = "1.0", UIMax = "200.0"))
-	float AttackTimeMsec;
+	// The amount of time to ramp into any dynamics processing effect
+	UPROPERTY(EditAnywhere, BlueprintReadWrite, Category = Response, meta = (DisplayName = "Attack Time (ms)", ClampMin = "1.0", ClampMax = "300.0", UIMin = "1.0", UIMax = "200.0", EditCondition = "!bBypass"))
+	float AttackTimeMsec = 10.0f;
 
-	UPROPERTY(EditAnywhere, BlueprintReadWrite, Category = "SourceEffect|Preset", meta = (ClampMin = "20.0", ClampMax = "5000.0", UIMin = "20.0", UIMax = "5000.0"))
-	float ReleaseTimeMsec;
+	// The amount of time to release the dynamics processing effect
+	UPROPERTY(EditAnywhere, BlueprintReadWrite, Category = Response, meta = (DisplayName = "Release Time (ms)", ClampMin = "20.0", ClampMax = "5000.0", UIMin = "20.0", UIMax = "5000.0", EditCondition = "!bBypass"))
+	float ReleaseTimeMsec = 100.0f;
 
-	UPROPERTY(EditAnywhere, BlueprintReadWrite, Category = "SourceEffect|Preset", meta = (ClampMin = "-60.0", ClampMax = "0.0", UIMin = "-60.0", UIMax = "0.0"))
-	float ThresholdDb;
+	// The threshold at which to perform a dynamics processing operation
+	UPROPERTY(EditAnywhere, BlueprintReadWrite, Category = Dynamics, meta = (DisplayName = "Threshold (dB)", ClampMin = "-60.0", ClampMax = "0.0", UIMin = "-60.0", UIMax = "0.0", EditCondition = "!bBypass"))
+	float ThresholdDb = -6.0f;
 
-	UPROPERTY(EditAnywhere, BlueprintReadWrite, Category = "SourceEffect|Preset", meta = (ClampMin = "1.0", ClampMax = "20.0", UIMin = "1.0", UIMax = "20.0"))
-	float Ratio;
+	// The dynamics processor ratio used for compression/expansion
+	UPROPERTY(EditAnywhere, BlueprintReadWrite, Category = Dynamics, meta = (ClampMin = "1.0", ClampMax = "20.0", UIMin = "1.0", UIMax = "20.0", EditCondition = "!bBypass"))
+	float Ratio = 1.5f;
 
-	UPROPERTY(EditAnywhere, BlueprintReadWrite, Category = "SourceEffect|Preset", meta = (ClampMin = "0.0", ClampMax = "20.0", UIMin = "0.0", UIMax = "20.0"))
-	float KneeBandwidthDb;
+	// The knee bandwidth of the processor to use
+	UPROPERTY(EditAnywhere, BlueprintReadWrite, Category = Dynamics, meta = (DisplayName = "Knee Bandwidth (dB)", ClampMin = "0.0", ClampMax = "20.0", UIMin = "0.0", UIMax = "20.0", EditCondition = "!bBypass"))
+	float KneeBandwidthDb = 10.0f;
 
-	UPROPERTY(EditAnywhere, BlueprintReadWrite, Category = "SourceEffect|Preset", meta = (ClampMin = "-12.0", ClampMax = "20.0", UIMin = "-12.0", UIMax = "20.0"))
-	float InputGainDb;
+	// The input gain of the dynamics processor
+	UPROPERTY(EditAnywhere, BlueprintReadWrite, Category = General, meta = (DisplayName = "Input Gain (dB)", ClampMin = "-12.0", ClampMax = "20.0", UIMin = "-12.0", UIMax = "20.0", EditCondition = "!bBypass"))
+	float InputGainDb = 0.0f;
 
-	UPROPERTY(EditAnywhere, BlueprintReadWrite, Category = "SourceEffect|Preset", meta = (ClampMin = "0.0", ClampMax = "20.0", UIMin = "0.0", UIMax = "20.0"))
-	float OutputGainDb;
+	// The output gain of the dynamics processor
+	UPROPERTY(EditAnywhere, BlueprintReadWrite, Category = General, meta = (DisplayName = "Output Gain (dB)", ClampMin = "0.0", ClampMax = "20.0", UIMin = "0.0", UIMax = "20.0", EditCondition = "!bBypass"))
+	float OutputGainDb = 0.0f;
 
-	UPROPERTY(EditAnywhere, BlueprintReadWrite, Category = "SourceEffect|Preset")
+	// Whether the left and right channels are linked when determining envelopes
+	UPROPERTY(EditAnywhere, BlueprintReadWrite, Category = Dynamics, meta = (EditCondition = "!bBypass"))
 	uint32 bStereoLinked : 1;
 
-	UPROPERTY(EditAnywhere, BlueprintReadWrite, Category = "SourceEffect|Preset")
+	// Toggles treating the attack and release envelopes as analog-style vs digital-style (Analog will respond a bit more naturally/slower)
+	UPROPERTY(EditAnywhere, BlueprintReadWrite, Category = Response, meta = (EditCondition = "!bBypass"))
 	uint32 bAnalogMode : 1;
 
+	// Whether or not to bypass effect
+	UPROPERTY(EditAnywhere, BlueprintReadWrite, Category = General, meta = (DisplayName = "Bypass", DisplayAfter = "DynamicsProcessorType"))
+	uint8 bBypass : 1;
+
 	FSourceEffectDynamicsProcessorSettings()
-		: DynamicsProcessorType(ESourceEffectDynamicsProcessorType::Compressor)
-		, PeakMode(ESourceEffectDynamicsPeakMode::RootMeanSquared)
-		, LookAheadMsec(3.0f)
-		, AttackTimeMsec(10.0f)
-		, ReleaseTimeMsec(100.0f)
-		, ThresholdDb(-6.0f)
-		, Ratio(1.5f)
-		, KneeBandwidthDb(10.0f)
-		, InputGainDb(0.0f)
-		, OutputGainDb(0.0f)
-		, bStereoLinked(true)
+		: bStereoLinked(true)
 		, bAnalogMode(true)
+		, bBypass(false)
 	{
 	}
 };
@@ -101,6 +108,7 @@ public:
 protected:
 
 	Audio::FDynamicsProcessor DynamicsProcessor;
+	bool bBypass = false;
 };
 
 
@@ -115,9 +123,9 @@ public:
 
 	virtual FColor GetPresetColor() const override { return FColor(218.0f, 199.0f, 11.0f); }
 
-	UFUNCTION(BlueprintCallable, Category = "Audio|Effects")
+	UFUNCTION(BlueprintCallable, Category = SourceEffectPresets)
 	void SetSettings(const FSourceEffectDynamicsProcessorSettings& InSettings);
 
-	UPROPERTY(EditAnywhere, BlueprintReadOnly, Category = "Audio|Effects")
+	UPROPERTY(EditAnywhere, BlueprintReadOnly, Category = SourceEffectPreset, meta = (ShowOnlyInnerProperties))
 	FSourceEffectDynamicsProcessorSettings Settings;
 };
