Fixing bit crusher serialization bug and improving the display of the bit crusher

[RN] Fix serialization issue when making and saving a new bitcrusher source effect

diff --git a/Engine/Plugins/Runtime/Synthesis/Source/Synthesis/Classes/SourceEffects/SourceEffectBitCrusher.h b/Engine/Plugins/Runtime/Synthesis/Source/Synthesis/Classes/SourceEffects/SourceEffectBitCrusher.h
--- a/Engine/Plugins/Runtime/Synthesis/Source/Synthesis/Classes/SourceEffects/SourceEffectBitCrusher.h
+++ b/Engine/Plugins/Runtime/Synthesis/Source/Synthesis/Classes/SourceEffects/SourceEffectBitCrusher.h
@@ -30,24 +30,18 @@ struct SYNTHESIS_API FSourceEffectBitCrusherSettings
 {
 	GENERATED_USTRUCT_BODY()
 
-	UPROPERTY(meta = (PropertyDeprecated))
-	float CrushedSampleRate = 8000.0f;
-
 	// The reduced frequency to use for the audio stream. 
-	UPROPERTY(EditAnywhere, BlueprintReadWrite, Category = "SourceEffect|Preset", meta = (DisplayName = "Sample Rate", AudioParam = "SampleRate", AudioParamClass = "SoundModulationParameterFrequency", ClampMin = "0.1", ClampMax = "96000.0", UIMin = "500.0", UIMax = "16000.0"))
+	UPROPERTY(EditAnywhere, BlueprintReadWrite, Category = BitCrusher, meta = (DisplayName = "Sample Rate", AudioParam = "SampleRate", AudioParamClass = "SoundModulationParameterFrequency", ClampMin = "0.1", ClampMax = "96000.0", UIMin = "500.0", UIMax = "16000.0"))
 	FSoundModulationDestinationSettings SampleRateModulation;
 
-	UPROPERTY(meta = (PropertyDeprecated))
-	float CrushedBits = 8.0f;
-
 	// The reduced bit depth to use for the audio stream
-	UPROPERTY(EditAnywhere, BlueprintReadWrite, Category = "SourceEffect|Preset", meta = (DisplayName = "Bit Depth", AudioParam = "BitDepth", AudioParamClass = "SoundModulationParameterScaled", ClampMin = "1.0", ClampMax = "24.0", UIMin = "1.0", UIMax = "16.0"))
+	UPROPERTY(EditAnywhere, BlueprintReadWrite, Category = BitCrusher, meta = (DisplayName = "Bit Depth", AudioParam = "BitDepth", AudioParamClass = "SoundModulationParameterScaled", ClampMin = "1.0", ClampMax = "24.0", UIMin = "1.0", UIMax = "16.0"))
 	FSoundModulationDestinationSettings BitModulation;
 
 	FSourceEffectBitCrusherSettings()
 	{
-		SampleRateModulation.Value = CrushedSampleRate;
-		BitModulation.Value = CrushedBits;
+		SampleRateModulation.Value = 8000.0f;
+		BitModulation.Value = 8.0f;
 	}
 };
 
@@ -88,35 +82,33 @@ public:
 
 	virtual void OnInit() override;
 
-	virtual void Serialize(FArchive& Ar) override;
-
 	virtual FColor GetPresetColor() const override { return FColor(196.0f, 185.0f, 121.0f); }
 
-	UFUNCTION(BlueprintCallable, Category = "Audio|Effects|BitCrusher")
+	UFUNCTION(BlueprintCallable, Category = SourceEffectPreset)
 	void SetBits(float Bits);
 
-	UFUNCTION(BlueprintCallable, Category = "Audio|Effects|BitCrusher")
+	UFUNCTION(BlueprintCallable, Category = SourceEffectPreset)
 	void SetBitModulator(const USoundModulatorBase* Modulator);
 
-	UFUNCTION(BlueprintCallable, Category = "Audio|Effects|BitCrusher")
+	UFUNCTION(BlueprintCallable, Category = SourceEffectPreset)
 	void SetBitModulators(const TSet<USoundModulatorBase*>& InModulators);
 
-	UFUNCTION(BlueprintCallable, Category = "Audio|Effects|BitCrusher")
+	UFUNCTION(BlueprintCallable, Category = SourceEffectPreset)
 	void SetSampleRate(float SampleRate);
 
-	UFUNCTION(BlueprintCallable, Category = "Audio|Effects|BitCrusher")
+	UFUNCTION(BlueprintCallable, Category = SourceEffectPreset)
 	void SetSampleRateModulator(const USoundModulatorBase* Modulator);
 
-	UFUNCTION(BlueprintCallable, Category = "Audio|Effects|BitCrusher")
+	UFUNCTION(BlueprintCallable, Category = SourceEffectPreset)
 	void SetSampleRateModulators(const TSet<USoundModulatorBase*>& InModulators);
 
 	// Sets just base (i.e. carrier) setting values without modifying modulation source references
-	UFUNCTION(BlueprintCallable, Category = "Audio|Effects|BitCrusher")
+	UFUNCTION(BlueprintCallable, Category = SourceEffectPreset)
 	void SetSettings(const FSourceEffectBitCrusherBaseSettings& Settings);
 
-	UFUNCTION(BlueprintCallable, Category = "Audio|Effects|BitCrusher")
+	UFUNCTION(BlueprintCallable, Category = SourceEffectPreset)
 	void SetModulationSettings(const FSourceEffectBitCrusherSettings& ModulationSettings);
 
-	UPROPERTY(EditAnywhere, BlueprintReadOnly, Category = "SourceEffect|Preset")
+	UPROPERTY(EditAnywhere, BlueprintReadOnly, Category = SourceEffectPreset, meta = (ShowOnlyInnerProperties))
 	FSourceEffectBitCrusherSettings Settings;
 };
diff --git a/Engine/Plugins/Runtime/Synthesis/Source/Synthesis/Private/SourceEffectBitCrusher.cpp b/Engine/Plugins/Runtime/Synthesis/Source/Synthesis/Private/SourceEffectBitCrusher.cpp
--- a/Engine/Plugins/Runtime/Synthesis/Source/Synthesis/Private/SourceEffectBitCrusher.cpp
+++ b/Engine/Plugins/Runtime/Synthesis/Source/Synthesis/Private/SourceEffectBitCrusher.cpp
@@ -83,28 +83,6 @@ void USourceEffectBitCrusherPreset::OnInit()
 	SetSampleRateModulators(TransformObjectPtrSet(Settings.SampleRateModulation.Modulators));
 }
 
-void USourceEffectBitCrusherPreset::Serialize(FArchive& Ar)
-{
-	Super::Serialize(Ar);
-
-#if WITH_EDITORONLY_DATA
-	if (Ar.IsLoading() || Ar.IsSaving())
-	{
-		if (Settings.CrushedBits >= 0.0f)
-		{
-			Settings.BitModulation.Value = Settings.CrushedBits;
-			Settings.CrushedBits = -1.0f;
-		}
-
-		if (Settings.CrushedSampleRate >= 0.0f)
-		{
-			Settings.SampleRateModulation.Value = Settings.CrushedSampleRate;
-			Settings.CrushedSampleRate = -1.0f;
-		}
-	}
-#endif // WITH_EDITORONLY_DATA
-}
-
 void USourceEffectBitCrusherPreset::SetBits(float InBits)
 {
 	UpdateSettings([NewBits = InBits](FSourceEffectBitCrusherSettings& OutSettings)
