Cleaning up source effect properties and display metadata

[RN] Polishing up source and submix effects uproperty displays, tooltips, and more.

diff --git a/Engine/Plugins/Runtime/Synthesis/Source/Synthesis/Classes/SourceEffects/SourceEffectChorus.h b/Engine/Plugins/Runtime/Synthesis/Source/Synthesis/Classes/SourceEffects/SourceEffectChorus.h
--- a/Engine/Plugins/Runtime/Synthesis/Source/Synthesis/Classes/SourceEffects/SourceEffectChorus.h
+++ b/Engine/Plugins/Runtime/Synthesis/Source/Synthesis/Classes/SourceEffects/SourceEffectChorus.h
@@ -62,27 +62,27 @@ struct SYNTHESIS_API FSourceEffectChorusSettings
 	float Spread = 0.0f;
 
 	// The depth of the chorus effect
-	UPROPERTY(EditAnywhere, BlueprintReadWrite, Category = "SourceEffect|Preset", meta = (DisplayName = "Depth", AudioParamClass = "SoundModulationParameter", ClampMin = "0.0", ClampMax = "1.0", UIMin = "0.0", UIMax = "1.0"))
+	UPROPERTY(EditAnywhere, BlueprintReadWrite, Category = Chorus, meta = (DisplayName = "Depth", AudioParamClass = "SoundModulationParameter", ClampMin = "0.0", ClampMax = "1.0", UIMin = "0.0", UIMax = "1.0"))
 	FSoundModulationDestinationSettings DepthModulation;
 
 	// The frequency of the chorus effect
-	UPROPERTY(EditAnywhere, BlueprintReadWrite, Category = "SourceEffect|Preset", meta = (DisplayName = "Frequency", AudioParam = "LowRateFrequency", AudioParamClass = "SoundModulationParameterFrequency", UIMin = "0.0", UIMax = "5.0"))
+	UPROPERTY(EditAnywhere, BlueprintReadWrite, Category = Chorus, meta = (DisplayName = "Frequency (hz)", AudioParam = "LowRateFrequency", AudioParamClass = "SoundModulationParameterFrequency", UIMin = "0.0", UIMax = "5.0"))
 	FSoundModulationDestinationSettings FrequencyModulation;
 
 	// The feedback of the chorus effect
-	UPROPERTY(EditAnywhere, BlueprintReadWrite, Category = "SourceEffect|Preset", meta = (DisplayName = "Feedback", AudioParamClass = "SoundModulationParameter", ClampMin = "0.0", ClampMax = "1.0", UIMin = "0.0", UIMax = "1.0"))
+	UPROPERTY(EditAnywhere, BlueprintReadWrite, Category = Chorus, meta = (DisplayName = "Feedback", AudioParamClass = "SoundModulationParameter", ClampMin = "0.0", ClampMax = "1.0", UIMin = "0.0", UIMax = "1.0"))
 	FSoundModulationDestinationSettings FeedbackModulation;
 
 	// The wet level of the chorus effect
-	UPROPERTY(EditAnywhere, BlueprintReadWrite, Category = "SourceEffect|Preset", meta = (DisplayName = "Wet Level", AudioParamClass = "SoundModulationParameterVolume", ClampMin = "0.0", ClampMax = "1.0", UIMin = "0.0", UIMax = "1.0"))
+	UPROPERTY(EditAnywhere, BlueprintReadWrite, Category = Chorus, meta = (DisplayName = "Wet Level", AudioParamClass = "SoundModulationParameterVolume", ClampMin = "0.0", ClampMax = "1.0", UIMin = "0.0", UIMax = "1.0"))
 	FSoundModulationDestinationSettings WetModulation;
 
 	// The dry level of the chorus effect
-	UPROPERTY(EditAnywhere, BlueprintReadWrite, Category = "SourceEffect|Preset", meta = (DisplayName = "Dry Level", AudioParamClass = "SoundModulationParameterVolume", ClampMin = "0.0", ClampMax = "1.0", UIMin = "0.0", UIMax = "1.0"))
+	UPROPERTY(EditAnywhere, BlueprintReadWrite, Category = Chorus, meta = (DisplayName = "Dry Level", AudioParamClass = "SoundModulationParameterVolume", ClampMin = "0.0", ClampMax = "1.0", UIMin = "0.0", UIMax = "1.0"))
 	FSoundModulationDestinationSettings DryModulation;
 
 	// The spread of the effect (larger means greater difference between left and right delay lines)
-	UPROPERTY(EditAnywhere, BlueprintReadWrite, Category = "SourceEffect|Preset", meta = (DisplayName = "Spread", AudioParamClass = "SoundModulationParameter", ClampMin = "0.0", ClampMax = "1.0", UIMin = "0.0", UIMax = "1.0"))
+	UPROPERTY(EditAnywhere, BlueprintReadWrite, Category = Chorus, meta = (DisplayName = "Spread", AudioParamClass = "SoundModulationParameter", ClampMin = "0.0", ClampMax = "1.0", UIMin = "0.0", UIMax = "1.0"))
 	FSoundModulationDestinationSettings SpreadModulation;
 
 	FSourceEffectChorusSettings()
@@ -208,6 +208,6 @@ public:
 	UFUNCTION(BlueprintCallable, Category = "Audio|Effects|Chorus")
 	void SetModulationSettings(const FSourceEffectChorusSettings& ModulationSettings);
 
-	UPROPERTY(EditAnywhere, BlueprintReadOnly, Category = "SourceEffect|Preset")
+	UPROPERTY(EditAnywhere, BlueprintReadOnly, Category = SourceEffectPreset, meta = (ShowOnlyInnerProperties))
 	FSourceEffectChorusSettings Settings;
 };
diff --git a/Engine/Plugins/Runtime/Synthesis/Source/Synthesis/Classes/SourceEffects/SourceEffectConvolutionReverb.h b/Engine/Plugins/Runtime/Synthesis/Source/Synthesis/Classes/SourceEffects/SourceEffectConvolutionReverb.h
--- a/Engine/Plugins/Runtime/Synthesis/Source/Synthesis/Classes/SourceEffects/SourceEffectConvolutionReverb.h
+++ b/Engine/Plugins/Runtime/Synthesis/Source/Synthesis/Classes/SourceEffects/SourceEffectConvolutionReverb.h
@@ -24,15 +24,15 @@ struct SYNTHESIS_API FSourceEffectConvolutionReverbSettings
 	float NormalizationVolumeDb;
 
 	// Controls how much of the wet signal is mixed into the output, in Decibels
-	UPROPERTY(EditAnywhere, BlueprintReadWrite, Category = SubmixEffectPreset, meta = (ClampMin = "-96.0", ClampMax = "0.0"));
+	UPROPERTY(EditAnywhere, BlueprintReadWrite, Category = ConvolutionReverb, meta = (DisplayName = "Wet Volume (dB)", ClampMin = "-96.0", ClampMax = "0.0", EditCondition = "!bBypass"));
 	float WetVolumeDb;
 
 	// Controls how much of the dry signal is mixed into the output, in Decibels
-	UPROPERTY(EditAnywhere, BlueprintReadWrite, Category = SubmixEffectPreset, meta = (ClampMin = "-96.0", ClampMax = "0.0"));
+	UPROPERTY(EditAnywhere, BlueprintReadWrite, Category = ConvolutionReverb, meta = (DisplayName = "Dry Volume (dB)", ClampMin = "-96.0", ClampMax = "0.0", EditCondition = "!bBypass"));
 	float DryVolumeDb;
 
 	/* If true, input audio is directly routed to output audio with applying any effect. */
-	UPROPERTY(EditAnywhere, BlueprintReadWrite, Category = SourceEffectPreset)
+	UPROPERTY(EditAnywhere, BlueprintReadWrite, Category = ConvolutionReverb)
 	bool bBypass;
 };
 
@@ -119,7 +119,7 @@ public:
 	TObjectPtr<UAudioImpulseResponse> ImpulseResponse;
 
 	/** ConvolutionReverbPreset Preset Settings. */
-	UPROPERTY(EditAnywhere, BlueprintReadWrite, BlueprintSetter = SetSettings, Category = SourceEffectPreset)
+	UPROPERTY(EditAnywhere, BlueprintReadWrite, BlueprintSetter = SetSettings, Category = SourceEffectPreset, meta = (ShowOnlyInnerProperties))
 	FSourceEffectConvolutionReverbSettings Settings;
 
 
diff --git a/Engine/Plugins/Runtime/Synthesis/Source/Synthesis/Classes/SourceEffects/SourceEffectEQ.h b/Engine/Plugins/Runtime/Synthesis/Source/Synthesis/Classes/SourceEffects/SourceEffectEQ.h
--- a/Engine/Plugins/Runtime/Synthesis/Source/Synthesis/Classes/SourceEffects/SourceEffectEQ.h
+++ b/Engine/Plugins/Runtime/Synthesis/Source/Synthesis/Classes/SourceEffects/SourceEffectEQ.h
@@ -14,19 +14,19 @@ struct SYNTHESIS_API FSourceEffectEQBand
 	GENERATED_USTRUCT_BODY()
 
 	// The cutoff frequency of the band
-	UPROPERTY(EditAnywhere, BlueprintReadWrite, Category = "SourceEffect|Preset", meta = (ClampMin = "0.0", ClampMax = "20000.0", UIMin = "0.0", UIMax = "15000.0"))
+	UPROPERTY(EditAnywhere, BlueprintReadWrite, Category = Equalizer, meta = (DisplayName = "Frequency (hz)", ClampMin = "20.0", ClampMax = "20000.0", UIMin = "0.0", UIMax = "15000.0", EditCondition = "!bEnabled"))
 	float Frequency;
 
 	// The bandwidth (in octaves) of the band
-	UPROPERTY(EditAnywhere, BlueprintReadWrite, Category = "SourceEffect|Preset", meta = (ClampMin = "0.0", ClampMax = "10.0", UIMin = "0.0", UIMax = "10.0"))
+	UPROPERTY(EditAnywhere, BlueprintReadWrite, Category = Equalizer, meta = (DisplayName = "Bandwidth (octaves)", ClampMin = "0.1", ClampMax = "2.0", UIMin = "0.1", UIMax = "2.0", EditCondition = "!bEnabled"))
 	float Bandwidth;
 
 	// The gain in decibels to apply to the eq band
-	UPROPERTY(EditAnywhere, BlueprintReadWrite, Category = "SourceEffect|Preset", meta = (DisplayName = "Gain (dB)", ClampMin = "-90.0", ClampMax = "20.0", UIMin = "-90.0", UIMax = "20.0"))
+	UPROPERTY(EditAnywhere, BlueprintReadWrite, Category = Equalizer, meta = (DisplayName = "Gain (dB)", ClampMin = "-90.0", ClampMax = "20.0", UIMin = "-90.0", UIMax = "20.0", EditCondition = "!bEnabled"))
 	float GainDb;
 
 	// Whether or not the band is enabled. Allows changing bands on the fly.
-	UPROPERTY(EditAnywhere, BlueprintReadWrite, Category = "SourceEffect|Preset")
+	UPROPERTY(EditAnywhere, BlueprintReadWrite, Category = Equalizer, meta = (DisplayName = "Bypass"))
 	uint32 bEnabled : 1;
 
 	FSourceEffectEQBand()
@@ -45,7 +45,7 @@ struct SYNTHESIS_API FSourceEffectEQSettings
 	GENERATED_USTRUCT_BODY()
 
 	// The EQ bands to use
-	UPROPERTY(EditAnywhere, BlueprintReadWrite, Category = "SourceEffect|Preset")
+	UPROPERTY(EditAnywhere, BlueprintReadWrite, Category = SourceEffectPreset)
 	TArray<FSourceEffectEQBand> EQBands;
 };
 
@@ -86,6 +86,6 @@ public:
 	UFUNCTION(BlueprintCallable, Category = "Audio|Effects")
 	void SetSettings(const FSourceEffectEQSettings& InSettings);
 
-	UPROPERTY(EditAnywhere, BlueprintReadOnly, Category = "Audio|Effects")
+	UPROPERTY(EditAnywhere, BlueprintReadOnly, Category = SourceEffectPreset, meta = (ShowOnlyInnerProperties))
 	FSourceEffectEQSettings Settings;
 };
diff --git a/Engine/Plugins/Runtime/Synthesis/Source/Synthesis/Classes/SourceEffects/SourceEffectEnvelopeFollower.h b/Engine/Plugins/Runtime/Synthesis/Source/Synthesis/Classes/SourceEffects/SourceEffectEnvelopeFollower.h
--- a/Engine/Plugins/Runtime/Synthesis/Source/Synthesis/Classes/SourceEffects/SourceEffectEnvelopeFollower.h
+++ b/Engine/Plugins/Runtime/Synthesis/Source/Synthesis/Classes/SourceEffects/SourceEffectEnvelopeFollower.h
@@ -24,19 +24,19 @@ struct SYNTHESIS_API FSourceEffectEnvelopeFollowerSettings
 	GENERATED_USTRUCT_BODY()
 
 	// The attack time of the envelope follower in milliseconds
-	UPROPERTY(EditAnywhere, BlueprintReadWrite, Category = "SourceEffect|Preset", meta = (ClampMin = "0.0", UIMin = "0.0"))
+	UPROPERTY(EditAnywhere, BlueprintReadWrite, Category = EnvelopeFollower, meta = (DisplayName = "Attack Time (ms)", ClampMin = "0.0", UIMin = "0.0"))
 	float AttackTime;
 
 	// The release time of the envelope follower in milliseconds
-	UPROPERTY(EditAnywhere, BlueprintReadWrite, Category = "SourceEffect|Preset", meta = (ClampMin = "0.0", UIMin = "0.0"))
+	UPROPERTY(EditAnywhere, BlueprintReadWrite, Category = EnvelopeFollower, meta = (DisplayName = "Release Time (ms)", ClampMin = "0.0", UIMin = "0.0"))
 	float ReleaseTime;
 
 	// The peak mode of the envelope follower
-	UPROPERTY(EditAnywhere, BlueprintReadWrite, Category = "SourceEffect|Preset")
+	UPROPERTY(EditAnywhere, BlueprintReadWrite, Category = EnvelopeFollower)
 	EEnvelopeFollowerPeakMode PeakMode;
 
 	// Whether or not the envelope follower is in analog mode
-	UPROPERTY(EditAnywhere, BlueprintReadWrite, Category = "SourceEffect|Preset")
+	UPROPERTY(EditAnywhere, BlueprintReadWrite, Category = EnvelopeFollower)
 	bool bIsAnalogMode;
 
 	FSourceEffectEnvelopeFollowerSettings()
@@ -144,6 +144,6 @@ public:
 	UFUNCTION(BlueprintCallable, Category = "Audio|Effects", meta = (WorldContext = "WorldContextObject"))
 	void UnregisterEnvelopeFollowerListener(UEnvelopeFollowerListener* EnvelopeFollowerListener);
 
-	UPROPERTY(EditAnywhere, BlueprintReadOnly, Category = "Audio|Effects")
+	UPROPERTY(EditAnywhere, BlueprintReadOnly, Category = SourceEffectPreset, meta = (ShowOnlyInnerProperties))
 	FSourceEffectEnvelopeFollowerSettings Settings;
 };
diff --git a/Engine/Plugins/Runtime/Synthesis/Source/Synthesis/Classes/SourceEffects/SourceEffectFilter.h b/Engine/Plugins/Runtime/Synthesis/Source/Synthesis/Classes/SourceEffects/SourceEffectFilter.h
--- a/Engine/Plugins/Runtime/Synthesis/Source/Synthesis/Classes/SourceEffects/SourceEffectFilter.h
+++ b/Engine/Plugins/Runtime/Synthesis/Source/Synthesis/Classes/SourceEffects/SourceEffectFilter.h
@@ -43,39 +43,39 @@ struct SYNTHESIS_API FSourceEffectFilterAudioBusModulationSettings
 	GENERATED_USTRUCT_BODY()
 
 	// Audio bus to use to modulate the filter
-	UPROPERTY(EditAnywhere, BlueprintReadWrite, Category = "SourceEffect|Preset|AudioBus")
+	UPROPERTY(EditAnywhere, BlueprintReadWrite, Category = Filter)
 	TObjectPtr<UAudioBus> AudioBus = nullptr;
 
 	// The amplitude envelope follower attack time (in milliseconds) on the audio bus.
-	UPROPERTY(EditAnywhere, BlueprintReadWrite, Category = "SourceEffect|Modulation Envelope", meta = (ClampMin = "0", UIMin = "0", UIMax = "2000.0", DisplayName = "Attack Time (ms)"))
+	UPROPERTY(EditAnywhere, BlueprintReadWrite, Category = Filter, meta = (ClampMin = "0", UIMin = "0", UIMax = "2000.0", DisplayName = "Attack Time (ms)"))
 	int32 EnvelopeFollowerAttackTimeMsec = 10;
 
 	// The amplitude envelope follower release time (in milliseconds) on the audio bus.
-	UPROPERTY(EditAnywhere, BlueprintReadWrite, Category = "SourceEffect|Modulation Envelope", meta = (ClampMin = "0", UIMin = "0", UIMax = "2000.0", DisplayName = "Release Time (ms)"))
+	UPROPERTY(EditAnywhere, BlueprintReadWrite, Category = Filter, meta = (ClampMin = "0", UIMin = "0", UIMax = "2000.0", DisplayName = "Release Time (ms)"))
 	int32 EnvelopeFollowerReleaseTimeMsec = 100;
 
 	// An amount to scale the envelope follower output to map to the modulation values.  
-	UPROPERTY(EditAnywhere, BlueprintReadWrite, Category = "SourceEffect|Modulation Envelope", meta = (ClampMin = "0"))
+	UPROPERTY(EditAnywhere, BlueprintReadWrite, Category = Filter, meta = (ClampMin = "0"))
 	float EnvelopeGainMultiplier = 1.0f;
 
 	// Which parameter to modulate.
-	UPROPERTY(EditAnywhere, Category = "SourceEffect|Filter Modulation Params", meta = (InlineCategoryProperty))
+	UPROPERTY(EditAnywhere, Category = ModulationParameters, meta = (InlineCategoryProperty))
 	ESourceEffectFilterParam FilterParam = ESourceEffectFilterParam::FilterFrequency;
 
 	// The frequency modulation value (in semitones from the filter frequency) to use when the envelope is quietest
-	UPROPERTY(EditAnywhere, BlueprintReadWrite, Category = "SourceEffect|Filter Modulation Params", meta = (EditCondition = "FilterParam == ESourceEffectFilterParam::FilterFrequency", DisplayName = "Min Frequency Modulation Amount", EditConditionHides))
+	UPROPERTY(EditAnywhere, BlueprintReadWrite, Category = ModulationParameters, meta = (EditCondition = "FilterParam == ESourceEffectFilterParam::FilterFrequency", DisplayName = "Min Frequency Modulation Amount", EditConditionHides))
 	float MinFrequencyModulation = -12.0f;
 
 	// The frequency modulation value (in semitones from the filter frequency) to use when the envelope is loudest
-	UPROPERTY(EditAnywhere, BlueprintReadWrite, Category = "SourceEffect|Filter Modulation Params", meta = (EditCondition = "FilterParam == ESourceEffectFilterParam::FilterFrequency", DisplayName = "Max Frequency Modulation Amount", EditConditionHides))
+	UPROPERTY(EditAnywhere, BlueprintReadWrite, Category = ModulationParameters, meta = (EditCondition = "FilterParam == ESourceEffectFilterParam::FilterFrequency", DisplayName = "Max Frequency Modulation Amount", EditConditionHides))
 	float MaxFrequencyModulation = 12.0f;
 
 	// The resonance modulation value to use when the envelope is quietest
-	UPROPERTY(EditAnywhere, BlueprintReadWrite, Category = "SourceEffect|Filter Modulation Params", meta = (ClampMin = "0.0", ClampMax = "10.0", UIMin = "0", UIMax = "10.0", EditCondition = "FilterParam == ESourceEffectFilterParam::FilterResonance", DisplayName = "Min Resonance Modulation Amount", EditConditionHides))
+	UPROPERTY(EditAnywhere, BlueprintReadWrite, Category = ModulationParameters, meta = (ClampMin = "0.0", ClampMax = "10.0", UIMin = "0", UIMax = "10.0", EditCondition = "FilterParam == ESourceEffectFilterParam::FilterResonance", DisplayName = "Min Resonance Modulation Amount", EditConditionHides))
 	float MinResonanceModulation = 0.2f;
 
 	// The resonance modulation value to use when the envelope is loudest
-	UPROPERTY(EditAnywhere, BlueprintReadWrite, Category = "SourceEffect|Filter Modulation Params", meta = (ClampMin = "0.0", ClampMax = "10.0", UIMin = "0", UIMax = "10.0", EditCondition = "FilterParam == ESourceEffectFilterParam::FilterResonance", DisplayName = "Max Resonance Modulation Amount", EditConditionHides))
+	UPROPERTY(EditAnywhere, BlueprintReadWrite, Category = ModulationParameters, meta = (ClampMin = "0.0", ClampMax = "10.0", UIMin = "0", UIMax = "10.0", EditCondition = "FilterParam == ESourceEffectFilterParam::FilterResonance", DisplayName = "Max Resonance Modulation Amount", EditConditionHides))
 	float MaxResonanceModulation = 3.0f;
 };
 
@@ -85,32 +85,24 @@ struct SYNTHESIS_API FSourceEffectFilterSettings
 	GENERATED_USTRUCT_BODY()
 
 	// The type of filter circuit to use.
-	UPROPERTY(EditAnywhere, BlueprintReadWrite, Category = "SourceEffect")
-	ESourceEffectFilterCircuit FilterCircuit;
+	UPROPERTY(EditAnywhere, BlueprintReadWrite, Category = Filter, meta = (DisplayName = "Circuit"))
+	ESourceEffectFilterCircuit FilterCircuit = ESourceEffectFilterCircuit::StateVariable;
 
 	// The type of filter to use.
-	UPROPERTY(EditAnywhere, BlueprintReadWrite, Category = "SourceEffect")
-	ESourceEffectFilterType FilterType;
+	UPROPERTY(EditAnywhere, BlueprintReadWrite, Category = Filter, meta = (DisplayName = "Type"))
+	ESourceEffectFilterType FilterType = ESourceEffectFilterType::LowPass;
 
 	// The filter cutoff frequency
-	UPROPERTY(EditAnywhere, BlueprintReadWrite, Category = "SourceEffect", meta = (ClampMin = "20.0", UIMin = "20.0", UIMax = "12000.0"))
-	float CutoffFrequency;
+	UPROPERTY(EditAnywhere, BlueprintReadWrite, Category = Filter, meta = (DisplayName = "Cutoff Frequency (hz)", ClampMin = "20.0", UIMin = "20.0", UIMax = "12000.0"))
+	float CutoffFrequency = 800.0f;
 
 	// The filter resonance.
-	UPROPERTY(EditAnywhere, BlueprintReadWrite, Category = "SourceEffect", meta = (ClampMin = "0.5", ClampMax = "10.0", UIMin = "0.5", UIMax = "10.0"))
-	float FilterQ;
+	UPROPERTY(EditAnywhere, BlueprintReadWrite, Category = Filter, meta = (DisplayName = "Resonance (Q)", ClampMin = "0.5", ClampMax = "10.0", UIMin = "0.5", UIMax = "10.0"))
+	float FilterQ = 2.0f;
 
 	// Audio bus settings to use to modulate the filter frequency cutoff (auto-wah) with arbitrary audio from an audio bus
-	UPROPERTY(EditAnywhere, BlueprintReadWrite, Category = "SourceEffect")
+	UPROPERTY(EditAnywhere, BlueprintReadWrite, Category = Filter, meta = (ShowOnlyInnerProperties))
 	TArray<FSourceEffectFilterAudioBusModulationSettings> AudioBusModulation;
-
-	FSourceEffectFilterSettings()
-		: FilterCircuit(ESourceEffectFilterCircuit::StateVariable)
-		, FilterType(ESourceEffectFilterType::LowPass)
-		, CutoffFrequency(800.0f)
-		, FilterQ(2.0f)
-	{
-	}
 };
 
 class SYNTHESIS_API FSourceEffectFilter : public FSoundEffectSource
@@ -180,6 +172,6 @@ public:
 	UFUNCTION(BlueprintCallable, Category = "Audio|Effects")
 	void SetSettings(const FSourceEffectFilterSettings& InSettings);
 
-	UPROPERTY(EditAnywhere, BlueprintReadOnly, Category = "Audio|Effects")
+	UPROPERTY(EditAnywhere, BlueprintReadOnly, Category = SourceEffectPreset, meta = (ShowOnlyInnerProperties))
 	FSourceEffectFilterSettings Settings;
 };
diff --git a/Engine/Plugins/Runtime/Synthesis/Source/Synthesis/Classes/SourceEffects/SourceEffectFoldbackDistortion.h b/Engine/Plugins/Runtime/Synthesis/Source/Synthesis/Classes/SourceEffects/SourceEffectFoldbackDistortion.h
--- a/Engine/Plugins/Runtime/Synthesis/Source/Synthesis/Classes/SourceEffects/SourceEffectFoldbackDistortion.h
+++ b/Engine/Plugins/Runtime/Synthesis/Source/Synthesis/Classes/SourceEffects/SourceEffectFoldbackDistortion.h
@@ -12,22 +12,16 @@ struct SYNTHESIS_API FSourceEffectFoldbackDistortionSettings
 	GENERATED_USTRUCT_BODY()
 
 	// The amount of gain to add to input to allow forcing the triggering of the threshold
-	UPROPERTY(EditAnywhere, BlueprintReadWrite, Category = "SourceEffect|Preset", meta = (DisplayName = "Input Gain (dB)", ClampMin = "0.0", ClampMax = "60.0", UIMin = "0.0", UIMax = "20.0"))
-	float InputGainDb;
+	UPROPERTY(EditAnywhere, BlueprintReadWrite, Category = Distortion, meta = (DisplayName = "Input Gain (dB)", ClampMin = "0.0", ClampMax = "60.0", UIMin = "0.0", UIMax = "20.0"))
+	float InputGainDb = 0.0f;
 
 	// If the audio amplitude is higher than this, it will fold back
-	UPROPERTY(EditAnywhere, BlueprintReadWrite, Category = "SourceEffect|Preset", meta = (DisplayName = "Threshold (dB)", ClampMin = "-90.0", ClampMax = "0.0", UIMin = "-60.0", UIMax = "0.0"))
-	float ThresholdDb;
+	UPROPERTY(EditAnywhere, BlueprintReadWrite, Category = Distortion, meta = (DisplayName = "Threshold (dB)", ClampMin = "-90.0", ClampMax = "0.0", UIMin = "-60.0", UIMax = "0.0"))
+	float ThresholdDb = -6.0f;
 
 	// The amount of gain to apply to the output
-	UPROPERTY(EditAnywhere, BlueprintReadWrite, Category = "SourceEffect|Preset", meta = (DisplayName = "Output Gain (dB)", ClampMin = "-90.0", ClampMax = "20.0", UIMin = "-60.0", UIMax = "20.0"))
-	float OutputGainDb;
-
-	FSourceEffectFoldbackDistortionSettings()
-		: InputGainDb(0.0f)
-		, ThresholdDb(-6.0f)
-		, OutputGainDb(-3.0f)
-	{}
+	UPROPERTY(EditAnywhere, BlueprintReadWrite, Category = Distortion, meta = (DisplayName = "Output Gain (dB)", ClampMin = "-90.0", ClampMax = "20.0", UIMin = "-60.0", UIMax = "20.0"))
+	float OutputGainDb = -3.0f;
 };
 
 class SYNTHESIS_API FSourceEffectFoldbackDistortion : public FSoundEffectSource
@@ -59,6 +53,6 @@ public:
 	UFUNCTION(BlueprintCallable, Category = "Audio|Effects")
 	void SetSettings(const FSourceEffectFoldbackDistortionSettings& InSettings);
 
-	UPROPERTY(EditAnywhere, BlueprintReadOnly, Category = "Audio|Effects")
+	UPROPERTY(EditAnywhere, BlueprintReadOnly, Category = SourceEffectPreset, meta = (ShowOnlyInnerProperties))
 	FSourceEffectFoldbackDistortionSettings Settings;
 };
diff --git a/Engine/Plugins/Runtime/Synthesis/Source/Synthesis/Classes/SourceEffects/SourceEffectMidSideSpreader.h b/Engine/Plugins/Runtime/Synthesis/Source/Synthesis/Classes/SourceEffects/SourceEffectMidSideSpreader.h
--- a/Engine/Plugins/Runtime/Synthesis/Source/Synthesis/Classes/SourceEffects/SourceEffectMidSideSpreader.h
+++ b/Engine/Plugins/Runtime/Synthesis/Source/Synthesis/Classes/SourceEffects/SourceEffectMidSideSpreader.h
@@ -25,27 +25,20 @@ struct SYNTHESIS_API FSourceEffectMidSideSpreaderSettings
 	GENERATED_USTRUCT_BODY()
 
 	// Amount of Mid/Side Spread. 0.0 is no spread, 1.0 is full wide. 
-	UPROPERTY(EditAnywhere, BlueprintReadWrite, Category = "SourceEffect|Preset", meta = (ClampMin = "0.0", ClampMax = "1.0", UIMin = "0.0", UIMax = "1.0"))
-	float SpreadAmount;
+	UPROPERTY(EditAnywhere, BlueprintReadWrite, Category = MidSideSpreader, meta = (ClampMin = "0.0", ClampMax = "1.0", UIMin = "0.0", UIMax = "1.0"))
+	float SpreadAmount = 0.5f;
 
 	// Indicate the channel mode of the input signal
-	UPROPERTY(EditAnywhere, BlueprintReadWrite, Category = "SourceEffect|Preset")
-	EStereoChannelMode InputMode;
+	UPROPERTY(EditAnywhere, BlueprintReadWrite, Category = MidSideSpreader)
+	EStereoChannelMode InputMode = EStereoChannelMode::LeftRight;
 
 	// Indicate the channel mode of the output signal
-	UPROPERTY(EditAnywhere, BlueprintReadWrite, Category = "SourceEffect|Preset")
-	EStereoChannelMode OutputMode;
+	UPROPERTY(EditAnywhere, BlueprintReadWrite, Category = MidSideSpreader)
+	EStereoChannelMode OutputMode = EStereoChannelMode::LeftRight;
 
 	// Indicate whether an equal power relationship between the mid and side channels should be maintained
-	UPROPERTY(EditAnywhere, BlueprintReadWrite, Category = "SourceEffect|Preset")
-	bool bEqualPower;
-
-	FSourceEffectMidSideSpreaderSettings()
-		: SpreadAmount(0.5f)
-		, InputMode(EStereoChannelMode::LeftRight)
-		, OutputMode(EStereoChannelMode::LeftRight)
-		, bEqualPower(false)
-	{}
+	UPROPERTY(EditAnywhere, BlueprintReadWrite, Category = MidSideSpreader)
+	bool bEqualPower = false;
 };
 
 // ========================================================================
@@ -102,6 +95,6 @@ public:
 	
 	// The copy of the settings struct. Can't be written to in BP, but can be read.
 	// Note that the value read in BP is the serialized settings, will not reflect dynamic changes made in BP.
-	UPROPERTY(EditAnywhere, BlueprintReadOnly, Category = "Audio|Effects")
+	UPROPERTY(EditAnywhere, BlueprintReadOnly, Category = SourceEffectPreset, meta = (ShowOnlyInnerProperties))
 	FSourceEffectMidSideSpreaderSettings Settings;
 };
diff --git a/Engine/Plugins/Runtime/Synthesis/Source/Synthesis/Classes/SourceEffects/SourceEffectMotionFilter.h b/Engine/Plugins/Runtime/Synthesis/Source/Synthesis/Classes/SourceEffects/SourceEffectMotionFilter.h
--- a/Engine/Plugins/Runtime/Synthesis/Source/Synthesis/Classes/SourceEffects/SourceEffectMotionFilter.h
+++ b/Engine/Plugins/Runtime/Synthesis/Source/Synthesis/Classes/SourceEffects/SourceEffectMotionFilter.h
@@ -32,19 +32,19 @@ UENUM(BlueprintType)
 enum class ESourceEffectMotionFilterModDestination : uint8
 {
 	// Filter input frequencies range between 20.0f and 15000.0f.
-	FilterACutoffFrequency = 0,
+	FilterACutoffFrequency = 0 UMETA(DisplayName = "Filter A Cutoff Frequency (hz)"),
 
 	// Filter input resonances range between 0.5f and 10.0f.
-	FilterAResonance,
+	FilterAResonance UMETA(DisplayName = "Filter A Resonance (Q)"),
 
 	// Filter output dB range between 10.0f and -96.0f. Final Filter output is clamped to +6 dB, use positive values with caution.
 	FilterAOutputVolumeDB UMETA(DisplayName = "Filter A Output Volume (dB)"),
 
 	// Filter input frequencies range between 20.0f and 15000.0f.
-	FilterBCutoffFrequency,
+	FilterBCutoffFrequency UMETA(DisplayName = "Filter B Cutoff Frequency (hz)"),
 
 	// Filter input resonances range between 0.5f and 10.0f.
-	FilterBResonance,
+	FilterBResonance UMETA(DisplayName = "Filter B Resonance (Q)"),
 
 	// Filter output dB range between 10.0f and -96.0f. Final Filter output is clamped to +6 dB, use positive values with caution.
 	FilterBOutputVolumeDB UMETA(DisplayName = "Filter B Output Volume (dB)"),
@@ -87,28 +87,20 @@ struct SYNTHESIS_API FSourceEffectIndividualFilterSettings
 	GENERATED_USTRUCT_BODY()
 
 	// The type of filter circuit to use.
-	UPROPERTY(EditAnywhere, BlueprintReadWrite, Category = "SourceEffect|Preset")
-	ESourceEffectMotionFilterCircuit FilterCircuit;
+	UPROPERTY(EditAnywhere, BlueprintReadWrite, Category = Filter, meta = (DisplayName = "Circuit"))
+	ESourceEffectMotionFilterCircuit FilterCircuit = ESourceEffectMotionFilterCircuit::Ladder;
 
 	// The type of filter to use.
-	UPROPERTY(EditAnywhere, BlueprintReadWrite, Category = "SourceEffect|Preset")
-	ESourceEffectMotionFilterType FilterType;
+	UPROPERTY(EditAnywhere, BlueprintReadWrite, Category = Filter, meta = (DisplayName = "Type"))
+	ESourceEffectMotionFilterType FilterType = ESourceEffectMotionFilterType::LowPass;
 
 	// The filter cutoff frequency
-	UPROPERTY(EditAnywhere, BlueprintReadWrite, Category = "SourceEffect|Preset", meta = (ClampMin = "20.0", UIMin = "20.0", UIMax = "12000.0"))
-	float CutoffFrequency;
+	UPROPERTY(EditAnywhere, BlueprintReadWrite, Category = Filter, meta = (DisplayName = "Cutoff Frequency (hz)", ClampMin = "20.0", UIMin = "20.0", UIMax = "12000.0"))
+	float CutoffFrequency = 800.0f;
 
 	// The filter resonance.
-	UPROPERTY(EditAnywhere, BlueprintReadWrite, Category = "SourceEffect|Preset", meta = (ClampMin = "0.5", ClampMax = "10.0", UIMin = "0.5", UIMax = "10.0"))
-	float FilterQ;
-
-	FSourceEffectIndividualFilterSettings()
-		: FilterCircuit(ESourceEffectMotionFilterCircuit::Ladder)
-		, FilterType(ESourceEffectMotionFilterType::LowPass)
-		, CutoffFrequency(800.0f)
-		, FilterQ(2.0f)
-	{
-	}
+	UPROPERTY(EditAnywhere, BlueprintReadWrite, Category = Filter, meta = (DisplayName = "Resonance (Q)", ClampMin = "0.5", ClampMax = "10.0", UIMin = "0.5", UIMax = "10.0"))
+	float FilterQ = 2.0f;
 };
 
 
@@ -118,33 +110,24 @@ struct SYNTHESIS_API FSourceEffectMotionFilterModulationSettings
 	GENERATED_USTRUCT_BODY()
 
 	// The Modulation Source
-	UPROPERTY(EditAnywhere, BlueprintReadWrite, Category = "SourceEffect|Preset")
-	ESourceEffectMotionFilterModSource ModulationSource;
+	UPROPERTY(EditAnywhere, BlueprintReadWrite, Category = MotionFilterModulation)
+	ESourceEffectMotionFilterModSource ModulationSource = ESourceEffectMotionFilterModSource::DistanceFromListener;
 
 	// The Modulation Clamped Input Range
-	UPROPERTY(EditAnywhere, BlueprintReadWrite, Category = "SourceEffect|Preset")
-	FVector2D ModulationInputRange;
+	UPROPERTY(EditAnywhere, BlueprintReadWrite, Category = MotionFilterModulation)
+	FVector2D ModulationInputRange = { 0.0f, 1.0f };
 
 	// The Modulation Random Minimum Output Range
-	UPROPERTY(EditAnywhere, BlueprintReadWrite, Category = "SourceEffect|Preset")
-	FVector2D ModulationOutputMinimumRange;
+	UPROPERTY(EditAnywhere, BlueprintReadWrite, Category = MotionFilterModulation)
+	FVector2D ModulationOutputMinimumRange = { 0.0f, 0.0f };
 
 	// The Modulation Random Maximum Output Range
-	UPROPERTY(EditAnywhere, BlueprintReadWrite, Category = "SourceEffect|Preset")
-	FVector2D ModulationOutputMaximumRange;
+	UPROPERTY(EditAnywhere, BlueprintReadWrite, Category = MotionFilterModulation)
+	FVector2D ModulationOutputMaximumRange  = { 1.0f, 1.0f };
 
 	// Update Ease Speed in milliseconds
-	UPROPERTY(EditAnywhere, BlueprintReadWrite, Category = "SourceEffect|Preset", meta = (ClampMin = "0.0", UIMin = "0.0"))
-	float UpdateEaseMS;
-
-	FSourceEffectMotionFilterModulationSettings()
-		: ModulationSource(ESourceEffectMotionFilterModSource::DistanceFromListener)
-		, ModulationInputRange(0.0f, 1.0f)
-		, ModulationOutputMinimumRange(0.0f, 0.0f)
-		, ModulationOutputMaximumRange(1.0f, 1.0f)
-		, UpdateEaseMS(50.0f)
-	{
-	}
+	UPROPERTY(EditAnywhere, BlueprintReadWrite, Category = MotionFilterModulation, meta = (DisplayName = "Update Ease (ms)", ClampMin = "0.0", UIMin = "0.0"))
+	float UpdateEaseMS = 50.0f;
 };
 
 // ========================================================================
@@ -158,33 +141,30 @@ struct SYNTHESIS_API FSourceEffectMotionFilterSettings
 	GENERATED_USTRUCT_BODY()
 
 	// In Serial Mode, Filter A will process then Filter B will process; in Parallel mode, Filter A and Filter B will process the dry input seprately, then be mixed together afterward.
-	UPROPERTY(EditAnywhere, BlueprintReadWrite, Category = "SourceEffect|Preset")
-	ESourceEffectMotionFilterTopology MotionFilterTopology;
+	UPROPERTY(EditAnywhere, BlueprintReadWrite, Category = MotionFilter)
+	ESourceEffectMotionFilterTopology MotionFilterTopology = ESourceEffectMotionFilterTopology::ParallelMode;
 
 	// Filter Mix controls the amount of each filter in the signal where -1.0f outputs Only Filter A, 0.0f is an equal balance between Filter A and B, and 1.0f outputs only Filter B. How this blend works depends on the Topology.
-	UPROPERTY(EditAnywhere, BlueprintReadWrite, Category = "SourceEffect|Preset", meta = (ClampMin = "-1.0", ClampMax = "1.0", UIMin = "-1.0", UIMax = "1.0"))
-	float MotionFilterMix;
+	UPROPERTY(EditAnywhere, BlueprintReadWrite, Category = MotionFilter, meta = (ClampMin = "-1.0", ClampMax = "1.0", UIMin = "-1.0", UIMax = "1.0"))
+	float MotionFilterMix = 0.0f;
 
 	// Initial settings for Filter A
-	UPROPERTY(EditAnywhere, BlueprintReadWrite, Category = "SourceEffect|Preset")
+	UPROPERTY(EditAnywhere, BlueprintReadWrite, Category = MotionFilter, meta = (DisplayName = "Filter A"))
 	FSourceEffectIndividualFilterSettings FilterASettings;
 
 	// Initial settings for Filter B
-	UPROPERTY(EditAnywhere, BlueprintReadWrite, Category = "SourceEffect|Preset")
+	UPROPERTY(EditAnywhere, BlueprintReadWrite, Category = MotionFilter, meta = (DisplayName = "Filter B"))
 	FSourceEffectIndividualFilterSettings FilterBSettings;
 
 	// Modulation Mappings
-	UPROPERTY(EditAnywhere, BlueprintReadWrite, Category = "SourceEffect|Preset")
+	UPROPERTY(EditAnywhere, BlueprintReadWrite, Category = MotionFilter)
 	TMap<ESourceEffectMotionFilterModDestination,FSourceEffectMotionFilterModulationSettings> ModulationMappings;
 
 	// Dry volume pass-through in dB. 
-	UPROPERTY(EditAnywhere, BlueprintReadWrite, Category = "SourceEffect|Preset", meta = (ClampMin = "-96.0", UIMin = "-96.0", UIMax = "10.0"))
-	float DryVolumeDb;
+	UPROPERTY(EditAnywhere, BlueprintReadWrite, Category = MotionFilter, meta = (DisplayName = "Dry Volume (dB)", ClampMin = "-96.0", UIMin = "-96.0", UIMax = "10.0"))
+	float DryVolumeDb = -96.0f;
 
 	FSourceEffectMotionFilterSettings()
-		: MotionFilterTopology(ESourceEffectMotionFilterTopology::ParallelMode)
-		, MotionFilterMix(0.0f)
-		, DryVolumeDb(-96.0f)
 	{
 		ModulationMappings.Empty((uint8)ESourceEffectMotionFilterModDestination::Count);
 	}
@@ -357,6 +337,6 @@ public:
 	
 	// The copy of the settings struct. Can't be written to in BP, but can be read.
 	// Note that the value read in BP is the serialized settings, will not reflect dynamic changes made in BP.
-	UPROPERTY(EditAnywhere, BlueprintReadOnly, Category = "Audio|Effects")
+	UPROPERTY(EditAnywhere, BlueprintReadOnly, Category = SourceEffectPreset, meta = (ShowOnlyInnerProperties))
 	FSourceEffectMotionFilterSettings Settings;
 };
diff --git a/Engine/Plugins/Runtime/Synthesis/Source/Synthesis/Classes/SourceEffects/SourceEffectPanner.h b/Engine/Plugins/Runtime/Synthesis/Source/Synthesis/Classes/SourceEffects/SourceEffectPanner.h
--- a/Engine/Plugins/Runtime/Synthesis/Source/Synthesis/Classes/SourceEffects/SourceEffectPanner.h
+++ b/Engine/Plugins/Runtime/Synthesis/Source/Synthesis/Classes/SourceEffects/SourceEffectPanner.h
@@ -11,17 +11,12 @@ struct SYNTHESIS_API FSourceEffectPannerSettings
 	GENERATED_USTRUCT_BODY()
 
 	// The spread of the source. 1.0 means left only in left channel, right only in right; 0.0 means both mixed, -1.0 means right and left channels are inverted.
-	UPROPERTY(EditAnywhere, BlueprintReadWrite, Category = "SourceEffect|Preset", meta = (ClampMin = "-1.0", ClampMax = "1.0", UIMin = "-1.0", UIMax = "1.0"))
-	float Spread;
+	UPROPERTY(EditAnywhere, BlueprintReadWrite, Category = Panner, meta = (ClampMin = "-1.0", ClampMax = "1.0", UIMin = "-1.0", UIMax = "1.0"))
+	float Spread = 1.0f;
 
 	// The pan of the source. -1.0 means left, 0.0 means center, 1.0 means right.
-	UPROPERTY(EditAnywhere, BlueprintReadWrite, Category = "SourceEffect|Preset", meta = (ClampMin = "-1.0", ClampMax = "1.0", UIMin = "-1.0", UIMax = "1.0"))
-	float Pan;
-
-	FSourceEffectPannerSettings()
-		: Spread(1.0f)
-		, Pan(0.0f)
-	{}
+	UPROPERTY(EditAnywhere, BlueprintReadWrite, Category = Panner, meta = (ClampMin = "-1.0", ClampMax = "1.0", UIMin = "-1.0", UIMax = "1.0"))
+	float Pan = 0.0f;
 };
 
 class SYNTHESIS_API FSourceEffectPanner : public FSoundEffectSource
@@ -67,6 +62,6 @@ public:
 	UFUNCTION(BlueprintCallable, Category = "Audio|Effects")
 	void SetSettings(const FSourceEffectPannerSettings& InSettings);
 	
-	UPROPERTY(EditAnywhere, BlueprintReadOnly, Category = "Audio|Effects")
+	UPROPERTY(EditAnywhere, BlueprintReadOnly, Category = SourceEffectPreset, meta = (ShowOnlyInnerProperties))
 	FSourceEffectPannerSettings Settings;
 };
diff --git a/Engine/Plugins/Runtime/Synthesis/Source/Synthesis/Classes/SourceEffects/SourceEffectPhaser.h b/Engine/Plugins/Runtime/Synthesis/Source/Synthesis/Classes/SourceEffects/SourceEffectPhaser.h
--- a/Engine/Plugins/Runtime/Synthesis/Source/Synthesis/Classes/SourceEffects/SourceEffectPhaser.h
+++ b/Engine/Plugins/Runtime/Synthesis/Source/Synthesis/Classes/SourceEffects/SourceEffectPhaser.h
@@ -25,32 +25,24 @@ struct SYNTHESIS_API FSourceEffectPhaserSettings
 	GENERATED_USTRUCT_BODY()
 
 	// The wet level of the phaser effect
-	UPROPERTY(EditAnywhere, BlueprintReadWrite, Category = "SourceEffect|Preset", meta = (ClampMin = "0.0", ClampMax = "1.0", UIMin = "0.0", UIMax = "1.0"))
-	float WetLevel;
+	UPROPERTY(EditAnywhere, BlueprintReadWrite, Category = Phaser, meta = (ClampMin = "0.0", ClampMax = "1.0", UIMin = "0.0", UIMax = "1.0"))
+	float WetLevel = 0.2f;
 
 	// The LFO frequency of the phaser effect
-	UPROPERTY(EditAnywhere, BlueprintReadWrite, Category = "SourceEffect|Preset", meta = (ClampMin = "0.0", UIMin = "0.0", UIMax = "5.0"))
-	float Frequency;
+	UPROPERTY(EditAnywhere, BlueprintReadWrite, Category = Phaser, meta = (ClampMin = "0.0", UIMin = "0.0", UIMax = "5.0"))
+	float Frequency = 2.0f;
 
 	// The feedback of the phaser effect
-	UPROPERTY(EditAnywhere, BlueprintReadWrite, Category = "SourceEffect|Preset", meta = (ClampMin = "0.0", ClampMax = "1.0", UIMin = "0.0", UIMax = "1.0"))
-	float Feedback;
+	UPROPERTY(EditAnywhere, BlueprintReadWrite, Category = Phaser, meta = (ClampMin = "0.0", ClampMax = "1.0", UIMin = "0.0", UIMax = "1.0"))
+	float Feedback = 0.3f;
 
 	// The phaser LFO type
-	UPROPERTY(EditAnywhere, BlueprintReadWrite, Category = "SourceEffect|Preset", meta = (ClampMin = "0.0", ClampMax = "1.0", UIMin = "0.0", UIMax = "1.0"))
-	EPhaserLFOType LFOType;
+	UPROPERTY(EditAnywhere, BlueprintReadWrite, Category = Phaser, meta = (ClampMin = "0.0", ClampMax = "1.0", UIMin = "0.0", UIMax = "1.0"))
+	EPhaserLFOType LFOType = EPhaserLFOType::Sine;
 
 	// Whether or not to use quadtrature phase for the LFO modulation
-	UPROPERTY(EditAnywhere, BlueprintReadWrite, Category = "SourceEffect|Preset", meta = (ClampMin = "0.0", ClampMax = "1.0", UIMin = "0.0", UIMax = "1.0"))
-	bool UseQuadraturePhase;
-
-	FSourceEffectPhaserSettings()
-		: WetLevel(0.2f)
-		, Frequency(2.0f)
-		, Feedback(0.3f)
-		, LFOType(EPhaserLFOType::Sine)
-		, UseQuadraturePhase(false)
-	{}
+	UPROPERTY(EditAnywhere, BlueprintReadWrite, Category = Phaser, meta = (ClampMin = "0.0", ClampMax = "1.0", UIMin = "0.0", UIMax = "1.0"))
+	bool UseQuadraturePhase = false;
 };
 
 class SYNTHESIS_API FSourceEffectPhaser : public FSoundEffectSource
@@ -83,6 +75,6 @@ public:
 	void SetSettings(const FSourceEffectPhaserSettings& InSettings);
 
 	// The depth of the chorus effect
-	UPROPERTY(EditAnywhere, BlueprintReadWrite, Category = "SourceEffect|Preset")
+	UPROPERTY(EditAnywhere, BlueprintReadWrite, Category = SourceEffectPreset, meta = (ShowOnlyInnerProperties))
 	FSourceEffectPhaserSettings Settings;
 };
\ No newline at end of file
diff --git a/Engine/Plugins/Runtime/Synthesis/Source/Synthesis/Classes/SourceEffects/SourceEffectRingModulation.h b/Engine/Plugins/Runtime/Synthesis/Source/Synthesis/Classes/SourceEffects/SourceEffectRingModulation.h
--- a/Engine/Plugins/Runtime/Synthesis/Source/Synthesis/Classes/SourceEffects/SourceEffectRingModulation.h
+++ b/Engine/Plugins/Runtime/Synthesis/Source/Synthesis/Classes/SourceEffects/SourceEffectRingModulation.h
@@ -24,37 +24,28 @@ struct SYNTHESIS_API FSourceEffectRingModulationSettings
 	GENERATED_USTRUCT_BODY()
 
 	// Ring modulation modulator oscillator type
-	UPROPERTY(EditAnywhere, BlueprintReadWrite, Category = "SourceEffect|Preset", meta = (ClampMin = "5.0", UIMin = "10.0", UIMax = "5000.0"))
-	ERingModulatorTypeSourceEffect ModulatorType;
+	UPROPERTY(EditAnywhere, BlueprintReadWrite, Category = RingModulation, meta = (ClampMin = "5.0", UIMin = "10.0", UIMax = "5000.0"))
+	ERingModulatorTypeSourceEffect ModulatorType = ERingModulatorTypeSourceEffect::Sine;
 
 	// Ring modulation frequency
-	UPROPERTY(EditAnywhere, BlueprintReadWrite, Category = "SourceEffect|Preset", meta = (ClampMin = "5.0", UIMin = "10.0", UIMax = "5000.0"))
-	float Frequency;
+	UPROPERTY(EditAnywhere, BlueprintReadWrite, Category = RingModulation, meta = (DisplayName = "Frequency (hz)", ClampMin = "5.0", UIMin = "10.0", UIMax = "5000.0"))
+	float Frequency = 100.0f;
 
 	// Ring modulation depth
-	UPROPERTY(EditAnywhere, BlueprintReadWrite, Category = "SourceEffect|Preset", meta = (ClampMin = "0.0", ClampMax = "1.0", UIMin = "0.0", UIMax = "1.0"))
-	float Depth;
+	UPROPERTY(EditAnywhere, BlueprintReadWrite, Category = RingModulation, meta = (ClampMin = "0.0", ClampMax = "1.0", UIMin = "0.0", UIMax = "1.0"))
+	float Depth = 0.5f;
 
 	// Gain for the dry signal (no ring mod)
-	UPROPERTY(EditAnywhere, BlueprintReadWrite, Category = "SourceEffect|Preset", meta = (ClampMin = "0.0", ClampMax = "1.0", UIMin = "0.0", UIMax = "1.0"))
-	float DryLevel;
+	UPROPERTY(EditAnywhere, BlueprintReadWrite, Category = RingModulation, meta = (ClampMin = "0.0", ClampMax = "1.0", UIMin = "0.0", UIMax = "1.0"))
+	float DryLevel = 0.0f;
 
 	// Gain for the wet signal (with ring mod)
-	UPROPERTY(EditAnywhere, BlueprintReadWrite, Category = "SourceEffect|Preset", meta = (ClampMin = "0.0", ClampMax = "1.0", UIMin = "0.0", UIMax = "1.0"))
-	float WetLevel;
+	UPROPERTY(EditAnywhere, BlueprintReadWrite, Category = RingModulation, meta = (ClampMin = "0.0", ClampMax = "1.0", UIMin = "0.0", UIMax = "1.0"))
+	float WetLevel = 1.0f;
 
 	// Audio bus to use to modulate the effect
-	UPROPERTY(EditAnywhere, BlueprintReadWrite, Category = "SourceEffect|Preset")
+	UPROPERTY(EditAnywhere, BlueprintReadWrite, Category = RingModulation)
 	TObjectPtr<UAudioBus> AudioBusModulator;
-
-	FSourceEffectRingModulationSettings()
-		: ModulatorType(ERingModulatorTypeSourceEffect::Sine)
-		, Frequency(100.0f)
-		, Depth(0.5f)
-		, DryLevel(0.0f)
-		, WetLevel(1.0f)
-		, AudioBusModulator(nullptr)
-	{}
 };
 
 class SYNTHESIS_API FSourceEffectRingModulation : public FSoundEffectSource
@@ -88,6 +79,6 @@ public:
 	UFUNCTION(BlueprintCallable, Category = "Audio|Effects")
 	void SetSettings(const FSourceEffectRingModulationSettings& InSettings);
 
-	UPROPERTY(EditAnywhere, BlueprintReadOnly, Category = "Audio|Effects")
+	UPROPERTY(EditAnywhere, BlueprintReadOnly, Category = SourceEffectPreset, meta = (ShowOnlyInnerProperties))
 	FSourceEffectRingModulationSettings Settings;
 };
diff --git a/Engine/Plugins/Runtime/Synthesis/Source/Synthesis/Classes/SourceEffects/SourceEffectSimpleDelay.h b/Engine/Plugins/Runtime/Synthesis/Source/Synthesis/Classes/SourceEffects/SourceEffectSimpleDelay.h
--- a/Engine/Plugins/Runtime/Synthesis/Source/Synthesis/Classes/SourceEffects/SourceEffectSimpleDelay.h
+++ b/Engine/Plugins/Runtime/Synthesis/Source/Synthesis/Classes/SourceEffects/SourceEffectSimpleDelay.h
@@ -12,40 +12,35 @@ struct SYNTHESIS_API FSourceEffectSimpleDelaySettings
 	GENERATED_USTRUCT_BODY()
 
 	// Speed of sound in meters per second when using distance-based delay
-	UPROPERTY(EditAnywhere, BlueprintReadWrite, Category = "SourceEffect|Preset", meta = (ClampMin = "0.01", UIMin = "0.01", UIMax = "10000.0", EditCondition = "bDelayBasedOnDistance"))
-	float SpeedOfSound;
+	UPROPERTY(EditAnywhere, BlueprintReadWrite, Category = Delay, meta = (DisplayAfter = "bDelayBasedOnDistance", ClampMin = "0.01", UIMin = "0.01", UIMax = "10000.0", EditCondition = "bDelayBasedOnDistance"))
+	float SpeedOfSound = 343.0f;
 
 	// Delay amount in seconds
-	UPROPERTY(EditAnywhere, BlueprintReadWrite, Category = "SourceEffect|Preset", meta = (ClampMin = "0.0", ClampMax = "2.0", UIMin = "0.0", UIMax = "2.0"))
-	float DelayAmount;
+	UPROPERTY(EditAnywhere, BlueprintReadWrite, Category = Delay, meta = (DisplayAfter = "SpeedOfSound", ClampMin = "0.0", ClampMax = "2.0", UIMin = "0.0", UIMax = "2.0", EditCondition = "!bDelayBasedOnDistance"))
+	float DelayAmount = 0.0f;
 
 	// Gain stage on dry (non-delayed signal)
-	UPROPERTY(EditAnywhere, BlueprintReadWrite, Category = "SourceEffect|Preset", meta = (ClampMin = "0.0", ClampMax = "1.0", UIMin = "0.0", UIMax = "1.0"))
-	float DryAmount;
+	UPROPERTY(EditAnywhere, BlueprintReadWrite, Category = Delay, meta = (ClampMin = "0.0", ClampMax = "1.0", UIMin = "0.0", UIMax = "1.0"))
+	float DryAmount = 0.0f;
 
 	// Gain stage on wet (delayed) signal
-	UPROPERTY(EditAnywhere, BlueprintReadWrite, Category = "SourceEffect|Preset", meta = (ClampMin = "0.0", ClampMax = "1.0", UIMin = "0.0", UIMax = "1.0"))
-	float WetAmount;
+	UPROPERTY(EditAnywhere, BlueprintReadWrite, Category = Delay, meta = (ClampMin = "0.0", ClampMax = "1.0", UIMin = "0.0", UIMax = "1.0"))
+	float WetAmount = 1.0f;
 
 	// Amount to feedback into the delay line (because why not)
-	UPROPERTY(EditAnywhere, BlueprintReadWrite, Category = "SourceEffect|Preset", meta = (ClampMin = "0.0", ClampMax = "1.0", UIMin = "0.0", UIMax = "1.0"))
-	float Feedback;
+	UPROPERTY(EditAnywhere, BlueprintReadWrite, Category = Delay, meta = (ClampMin = "0.0", ClampMax = "1.0", UIMin = "0.0", UIMax = "1.0"))
+	float Feedback = 0.0f;
 
 	// Whether or not to delay the audio based on the distance to the listener or use manual delay
-	UPROPERTY(EditAnywhere, BlueprintReadWrite, Category = "SourceEffect|Preset")
+	UPROPERTY(EditAnywhere, BlueprintReadWrite, Category = Delay)
 	uint32 bDelayBasedOnDistance : 1;
 
 	// Whether or not to allow the attenuation distance override value vs the distance to listener to be used for distance-based delay.
-	UPROPERTY(EditAnywhere, BlueprintReadWrite, Category = "SourceEffect|Preset")
+	UPROPERTY(EditAnywhere, BlueprintReadWrite, Category = Delay)
 	uint32 bUseDistanceOverride : 1;
 	
 	FSourceEffectSimpleDelaySettings()
-		: SpeedOfSound(343.0f)
-		, DelayAmount(0.0f)
-		, DryAmount(0.0f)
-		, WetAmount(1.0f)
-		, Feedback(0.0f)
-		, bDelayBasedOnDistance(true)
+		: bDelayBasedOnDistance(true)
 		, bUseDistanceOverride(true)
 	{}
 };
@@ -83,6 +78,6 @@ public:
 	UFUNCTION(BlueprintCallable, Category = "Audio|Effects")
 	void SetSettings(const FSourceEffectSimpleDelaySettings& InSettings);
 
-	UPROPERTY(EditAnywhere, BlueprintReadOnly, Category = "Audio|Effects")
+	UPROPERTY(EditAnywhere, BlueprintReadOnly, Category = SourceEffectPreset, meta = (ShowOnlyInnerProperties))
 	FSourceEffectSimpleDelaySettings Settings;
 };
diff --git a/Engine/Plugins/Runtime/Synthesis/Source/Synthesis/Classes/SourceEffects/SourceEffectStereoDelay.h b/Engine/Plugins/Runtime/Synthesis/Source/Synthesis/Classes/SourceEffects/SourceEffectStereoDelay.h
--- a/Engine/Plugins/Runtime/Synthesis/Source/Synthesis/Classes/SourceEffects/SourceEffectStereoDelay.h
+++ b/Engine/Plugins/Runtime/Synthesis/Source/Synthesis/Classes/SourceEffects/SourceEffectStereoDelay.h
@@ -42,43 +42,43 @@ struct SYNTHESIS_API FSourceEffectStereoDelaySettings
 	GENERATED_USTRUCT_BODY()
 
 	// What mode to set the stereo delay effect.
-	UPROPERTY(EditAnywhere, BlueprintReadWrite, Category = "SourceEffect|Preset")
+	UPROPERTY(EditAnywhere, BlueprintReadWrite, Category = StereoDelay, meta = (DisplayName = "Mode"))
 	EStereoDelaySourceEffect DelayMode = EStereoDelaySourceEffect::PingPong;
 
 	// The base amount of delay in the left and right channels of the delay line.
-	UPROPERTY(EditAnywhere, BlueprintReadWrite, Category = "SourceEffect|Preset", meta = (ClampMin = "0.0", ClampMax = "2000.0", UIMin = "0.0", UIMax = "2000.0"))
+	UPROPERTY(EditAnywhere, BlueprintReadWrite, Category = StereoDelay, meta = (DisplayName = "Delay Time (ms)", ClampMin = "0.0", ClampMax = "2000.0", UIMin = "0.0", UIMax = "2000.0"))
 	float DelayTimeMsec = 500.0f;
 
 	// The amount of audio to feedback into the delay line once the delay has been tapped.
-	UPROPERTY(EditAnywhere, BlueprintReadWrite, Category = "SourceEffect|Preset", meta = (ClampMin = "0.0", ClampMax = "1.0", UIMin = "0.0", UIMax = "1.0"))
+	UPROPERTY(EditAnywhere, BlueprintReadWrite, Category = StereoDelay, meta = (ClampMin = "0.0", ClampMax = "1.0", UIMin = "0.0", UIMax = "1.0"))
 	float Feedback = 0.1f;
 
 	// Delay spread for left and right channels. Allows left and right channels to have differential delay amounts. Useful for stereo channel decorrelation.
-	UPROPERTY(EditAnywhere, BlueprintReadWrite, Category = "SourceEffect|Preset", meta = (ClampMin = "-1.0", ClampMax = "1.0", UIMin = "-1.0", UIMax = "1.0"))
+	UPROPERTY(EditAnywhere, BlueprintReadWrite, Category = StereoDelay, meta = (ClampMin = "-1.0", ClampMax = "1.0", UIMin = "-1.0", UIMax = "1.0"))
 	float DelayRatio = 0.2f;
 
 	// The amount of delay effect to mix with the dry input signal into the effect.
-	UPROPERTY(EditAnywhere, BlueprintReadWrite, Category = "SourceEffect|Preset", meta = (ClampMin = "0.0", ClampMax = "1.0", UIMin = "0.0", UIMax = "1.0"))
+	UPROPERTY(EditAnywhere, BlueprintReadWrite, Category = StereoDelay, meta = (ClampMin = "0.0", ClampMax = "1.0", UIMin = "0.0", UIMax = "1.0"))
 	float WetLevel = 0.4f;
 
 	// The dry level of the effect
-	UPROPERTY(EditAnywhere, BlueprintReadWrite, Category = "SourceEffect|Preset", meta = (ClampMin = "0.0", ClampMax = "1.0", UIMin = "0.0", UIMax = "1.0"))
+	UPROPERTY(EditAnywhere, BlueprintReadWrite, Category = StereoDelay, meta = (ClampMin = "0.0", ClampMax = "1.0", UIMin = "0.0", UIMax = "1.0"))
 	float DryLevel = 1.0f;
 
 	// Whether or not to enable filtering
-	UPROPERTY(EditAnywhere, BlueprintReadWrite, Category = "SourceEffect|Preset")
+	UPROPERTY(EditAnywhere, BlueprintReadWrite, Category = Filter, meta = (DisplayName = "Enabled"))
 	bool bFilterEnabled = false;
 
 	// Filter type to feed feedback audio to
-	UPROPERTY(EditAnywhere, BlueprintReadWrite, Category = "SourceEffect|Preset", meta = (EditCondition = "bFilterEnabled"))
+	UPROPERTY(EditAnywhere, BlueprintReadWrite, Category = Filter, meta = (DisplayName = "Type", EditCondition = "bFilterEnabled"))
 	EStereoDelayFiltertype FilterType = EStereoDelayFiltertype::Lowpass;
 
 	// Cutoff frequency of the filter
-	UPROPERTY(EditAnywhere, BlueprintReadWrite, Category = "SourceEffect|Preset", meta = (ClampMin = "20.0", ClampMax = "20000", UIMin = "20.0", UIMax = "20000", EditCondition = "bFilterEnabled"))
+	UPROPERTY(EditAnywhere, BlueprintReadWrite, Category = Filter, meta = (DisplayName = "Cutoff Frequency (hz)", ClampMin = "20.0", ClampMax = "20000", UIMin = "20.0", UIMax = "20000", EditCondition = "bFilterEnabled"))
 	float FilterFrequency = 20000.0f;
 
 	// Q of the filter
-	UPROPERTY(EditAnywhere, BlueprintReadWrite, Category = "SourceEffect|Preset", meta = (ClampMin = "0.5", ClampMax = "10.0", UIMin = "0.5", UIMax = "10.0", EditCondition = "bFilterEnabled"))
+	UPROPERTY(EditAnywhere, BlueprintReadWrite, Category = Filter, meta = (DisplayName = "Resonance (Q)", ClampMin = "0.5", ClampMax = "10.0", UIMin = "0.5", UIMax = "10.0", EditCondition = "bFilterEnabled"))
 	float FilterQ = 2.0f;
 };
 
@@ -98,8 +98,6 @@ protected:
 	Audio::FDelayStereo DelayStereo;
 };
 
-
-
 UCLASS(ClassGroup = AudioSourceEffect, meta = (BlueprintSpawnableComponent))
 class SYNTHESIS_API USourceEffectStereoDelayPreset : public USoundEffectSourcePreset
 {
@@ -114,6 +112,6 @@ public:
 	UFUNCTION(BlueprintCallable, Category = "Audio|Effects")
 	void SetSettings(const FSourceEffectStereoDelaySettings& InSettings);
 
-	UPROPERTY(EditAnywhere, BlueprintReadOnly, Category = "Audio|Effects")
+	UPROPERTY(EditAnywhere, BlueprintReadOnly, Category = SourceEffectPreset, meta = (ShowOnlyInnerProperties))
 	FSourceEffectStereoDelaySettings Settings;
 };
diff --git a/Engine/Plugins/Runtime/Synthesis/Source/Synthesis/Classes/SourceEffects/SourceEffectWaveShaper.h b/Engine/Plugins/Runtime/Synthesis/Source/Synthesis/Classes/SourceEffects/SourceEffectWaveShaper.h
--- a/Engine/Plugins/Runtime/Synthesis/Source/Synthesis/Classes/SourceEffects/SourceEffectWaveShaper.h
+++ b/Engine/Plugins/Runtime/Synthesis/Source/Synthesis/Classes/SourceEffects/SourceEffectWaveShaper.h
@@ -12,17 +12,12 @@ struct SYNTHESIS_API FSourceEffectWaveShaperSettings
 	GENERATED_USTRUCT_BODY()
 
 	// The amount of wave shaping. 0.0 = no wave shaping.
-	UPROPERTY(EditAnywhere, BlueprintReadWrite, Category = "SourceEffect|Preset", meta = (ClampMin = "0.0", UIMin = "0.0", UIMax = "500.0"))
-	float Amount;
+	UPROPERTY(EditAnywhere, BlueprintReadWrite, Category = WaveShaper, meta = (ClampMin = "0.0", UIMin = "0.0", UIMax = "500.0"))
+	float Amount = 1.0f;
 
 	// The amount of wave shaping. 0.0 = no wave shaping.
-	UPROPERTY(EditAnywhere, BlueprintReadWrite, Category = "SourceEffect|Preset", meta = (ClampMin = "-60.0", ClampMax = "20.0", UIMin = "-60.0", UIMax = "20.0"))
-	float OutputGainDb;
-
-	FSourceEffectWaveShaperSettings()
-		: Amount(1.0f)
-		, OutputGainDb(0.0f)
-	{}
+	UPROPERTY(EditAnywhere, BlueprintReadWrite, Category = WaveShaper, meta = (DisplayName = "Output Gain (dB)", ClampMin = "-60.0", ClampMax = "20.0", UIMin = "-60.0", UIMax = "20.0"))
+	float OutputGainDb = 0.0f;
 };
 
 class SYNTHESIS_API FSourceEffectWaveShaper : public FSoundEffectSource
@@ -42,8 +37,6 @@ protected:
 	int32 NumChannels;
 };
 
-
-
 UCLASS(ClassGroup = AudioSourceEffect, meta = (BlueprintSpawnableComponent))
 class SYNTHESIS_API USourceEffectWaveShaperPreset : public USoundEffectSourcePreset
 {
@@ -57,6 +50,6 @@ public:
 	UFUNCTION(BlueprintCallable, Category = "Audio|Effects")
 	void SetSettings(const FSourceEffectWaveShaperSettings& InSettings);
 
-	UPROPERTY(EditAnywhere, BlueprintReadOnly, Category = "Audio|Effects")
+	UPROPERTY(EditAnywhere, BlueprintReadOnly, Category = SourceEffectPreset, meta = (ShowOnlyInnerProperties))
 	FSourceEffectWaveShaperSettings Settings;
 };
diff --git a/Engine/Plugins/Runtime/Synthesis/Source/Synthesis/Classes/SubmixEffects/SubmixEffectConvolutionReverb.h b/Engine/Plugins/Runtime/Synthesis/Source/Synthesis/Classes/SubmixEffects/SubmixEffectConvolutionReverb.h
--- a/Engine/Plugins/Runtime/Synthesis/Source/Synthesis/Classes/SubmixEffects/SubmixEffectConvolutionReverb.h
+++ b/Engine/Plugins/Runtime/Synthesis/Source/Synthesis/Classes/SubmixEffects/SubmixEffectConvolutionReverb.h
@@ -24,11 +24,11 @@ struct SYNTHESIS_API FSubmixEffectConvolutionReverbSettings
 	float NormalizationVolumeDb;
 
 	// Controls how much of the wet signal is mixed into the output, in Decibels
-	UPROPERTY(EditAnywhere, BlueprintReadWrite, Category = SubmixEffectPreset, meta = (ClampMin = "-96.0", ClampMax = "0.0"));
+	UPROPERTY(EditAnywhere, BlueprintReadWrite, Category = SubmixEffectPreset, meta = (DisplayName = "Wet Volume (dB)", ClampMin = "-96.0", ClampMax = "0.0"));
 	float WetVolumeDb;
 
 	// Controls how much of the dry signal is mixed into the output, in Decibels
-	UPROPERTY(EditAnywhere, BlueprintReadWrite, Category = SubmixEffectPreset, meta = (ClampMin = "-96.0", ClampMax = "0.0"));
+	UPROPERTY(EditAnywhere, BlueprintReadWrite, Category = SubmixEffectPreset, meta = (DisplayName = "Dry Volume (dB)", ClampMin = "-96.0", ClampMax = "0.0"));
 	float DryVolumeDb;
 
 	/* If true, input audio is directly routed to output audio with applying any effect. */
@@ -155,7 +155,7 @@ public:
 	TObjectPtr<UAudioImpulseResponse> ImpulseResponse;
 
 	/** ConvolutionReverbPreset Preset Settings. */
-	UPROPERTY(EditAnywhere, BlueprintReadWrite, BlueprintSetter = SetSettings, Category = SubmixEffectPreset)
+	UPROPERTY(EditAnywhere, BlueprintReadWrite, BlueprintSetter = SetSettings, Category = SubmixEffectPreset, meta = (ShowOnlyInnerProperties))
 	FSubmixEffectConvolutionReverbSettings Settings;
 
 
diff --git a/Engine/Plugins/Runtime/Synthesis/Source/Synthesis/Classes/SubmixEffects/SubmixEffectDelay.h b/Engine/Plugins/Runtime/Synthesis/Source/Synthesis/Classes/SubmixEffects/SubmixEffectDelay.h
--- a/Engine/Plugins/Runtime/Synthesis/Source/Synthesis/Classes/SubmixEffects/SubmixEffectDelay.h
+++ b/Engine/Plugins/Runtime/Synthesis/Source/Synthesis/Classes/SubmixEffects/SubmixEffectDelay.h
@@ -20,23 +20,16 @@ struct SYNTHESIS_API FSubmixEffectDelaySettings
 	GENERATED_USTRUCT_BODY()
 
 	// Maximum possible length for a delay, in milliseconds. Changing this at runtime will reset the effect.
-	UPROPERTY(EditAnywhere, BlueprintReadWrite, Category = Initialization, meta = (ClampMin = "10.0", UIMin = "10.0", UIMax = "20000.0"))
-	float MaximumDelayLength;
+	UPROPERTY(EditAnywhere, BlueprintReadWrite, Category = Initialization, meta = (DisplayName = "Maximum Delay Length (ms)", ClampMin = "10.0", UIMin = "10.0", UIMax = "20000.0"))
+	float MaximumDelayLength = 2000.0f;
 
 	// Number of milliseconds over which a tap will reach it's set length and gain. Smaller values are more responsive, while larger values will make pitching less dramatic.
-	UPROPERTY(EditAnywhere, BlueprintReadWrite, Category = Realtime, meta = (ClampMin = "0.0", UIMin = "0.0", UIMax = "20000.0"))
-	float InterpolationTime;
+	UPROPERTY(EditAnywhere, BlueprintReadWrite, Category = Realtime, meta = (DisplayName = "Interpolation Time (ms)", ClampMin = "0.0", UIMin = "0.0", UIMax = "20000.0"))
+	float InterpolationTime = 400.0f;
 
 	// Number of milliseconds of delay.  Caps at max delay at runtime.
-	UPROPERTY(EditAnywhere, BlueprintReadWrite, Category = Realtime, meta = (ClampMin = "0.0", UIMin = "0.0", UIMax = "20000.0"))
-	float DelayLength;
-
-	FSubmixEffectDelaySettings()
-		: MaximumDelayLength(2000.0f)
-		, InterpolationTime(400.0f)
-		, DelayLength(1000.0f)
-	{
-	}
+	UPROPERTY(EditAnywhere, BlueprintReadWrite, Category = Realtime, meta = (DisplayName = "Delay Length (ms)", ClampMin = "0.0", UIMin = "0.0", UIMax = "20000.0"))
+	float DelayLength = 1000.0f;
 };
 
 UCLASS()
diff --git a/Engine/Plugins/Runtime/Synthesis/Source/Synthesis/Classes/SubmixEffects/SubmixEffectFilter.h b/Engine/Plugins/Runtime/Synthesis/Source/Synthesis/Classes/SubmixEffects/SubmixEffectFilter.h
--- a/Engine/Plugins/Runtime/Synthesis/Source/Synthesis/Classes/SubmixEffects/SubmixEffectFilter.h
+++ b/Engine/Plugins/Runtime/Synthesis/Source/Synthesis/Classes/SubmixEffects/SubmixEffectFilter.h
@@ -36,28 +36,20 @@ struct SYNTHESIS_API FSubmixEffectFilterSettings
 	GENERATED_USTRUCT_BODY()
 
 	// What type of filter to use for the submix filter effect
-	UPROPERTY(EditAnywhere, BlueprintReadWrite, Category = Filter)
-	ESubmixFilterType FilterType;
+	UPROPERTY(EditAnywhere, BlueprintReadWrite, Category = Filter, meta = (DisplayName = "Type"))
+	ESubmixFilterType FilterType = ESubmixFilterType::LowPass;
 
 	// What type of filter algorithm to use for the submix filter effect
-	UPROPERTY(EditAnywhere, BlueprintReadWrite, Category = Filter)
-	ESubmixFilterAlgorithm FilterAlgorithm;
+	UPROPERTY(EditAnywhere, BlueprintReadWrite, Category = Filter, meta = (DisplayName = "Circuit"))
+	ESubmixFilterAlgorithm FilterAlgorithm = ESubmixFilterAlgorithm::OnePole;
 
 	// The output filter cutoff frequency (hz) [0.0, 20000.0]
-	UPROPERTY(EditAnywhere, BlueprintReadWrite, Category = Filter, meta = (ClampMin = "0.0", ClampMax = "20000.0", UIMin = "0.0", UIMax = "20000.0"))
-	float FilterFrequency;
+	UPROPERTY(EditAnywhere, BlueprintReadWrite, Category = Filter, meta = (DisplayName = "Cutoff Frequency (hz)", ClampMin = "0.0", ClampMax = "20000.0", UIMin = "0.0", UIMax = "20000.0"))
+	float FilterFrequency = 20000.0f;
 
 	// The output filter resonance (Q) [0.5, 10]
-	UPROPERTY(EditAnywhere, BlueprintReadWrite, Category = Filter, meta = (ClampMin = "0.5", ClampMax = "10.0", UIMin = "0.5", UIMax = "10.0"))
-	float FilterQ;
-
-	FSubmixEffectFilterSettings()
-		: FilterType(ESubmixFilterType::LowPass)
-		, FilterAlgorithm(ESubmixFilterAlgorithm::OnePole)
-		, FilterFrequency(20000.0f)
-		, FilterQ(2.0f)
-	{
-	}
+	UPROPERTY(EditAnywhere, BlueprintReadWrite, Category = Filter, meta = (DisplayName = "Resonance (Q)", ClampMin = "0.5", ClampMax = "10.0", UIMin = "0.5", UIMax = "10.0"))
+	float FilterQ = 2.0f;
 };
 
 class SYNTHESIS_API FSubmixEffectFilter : public FSoundEffectSubmix
@@ -164,6 +156,6 @@ public:
 
 public:
 
-	UPROPERTY(EditAnywhere, BlueprintReadWrite, Category = SubmixEffectPreset, Meta = (ShowOnlyInnerProperties))
+	UPROPERTY(EditAnywhere, BlueprintReadWrite, Category = SubmixEffectPreset, meta = (ShowOnlyInnerProperties))
 	FSubmixEffectFilterSettings Settings;
 };
diff --git a/Engine/Plugins/Runtime/Synthesis/Source/Synthesis/Classes/SubmixEffects/SubmixEffectFlexiverb.h b/Engine/Plugins/Runtime/Synthesis/Source/Synthesis/Classes/SubmixEffects/SubmixEffectFlexiverb.h
--- a/Engine/Plugins/Runtime/Synthesis/Source/Synthesis/Classes/SubmixEffects/SubmixEffectFlexiverb.h
+++ b/Engine/Plugins/Runtime/Synthesis/Source/Synthesis/Classes/SubmixEffects/SubmixEffectFlexiverb.h
@@ -12,11 +12,11 @@ struct SYNTHESIS_API FSubmixEffectFlexiverbSettings
 	GENERATED_USTRUCT_BODY()
 	
 	/** PreDelay - 0.01 < 10.0 < 40.0 - Amount of delay to the first echo in milliseconds. */
-	UPROPERTY(EditAnywhere, BlueprintReadWrite, Category = ReverbParameters, meta = (ClampMin = "0.01", ClampMax = "30"))
+	UPROPERTY(EditAnywhere, BlueprintReadWrite, Category = ReverbParameters, meta = (DisplayName = "Pre Delay (ms)", ClampMin = "0.01", ClampMax = "30"))
 	float PreDelay;
 
 	/** Time in seconds it will take for the impulse response to decay to -60 dB. */
-	UPROPERTY(EditAnywhere, BlueprintReadWrite, Category = ReverbParameters, meta = (ClampMin = "0.4", ClampMax = "5.0"))
+	UPROPERTY(EditAnywhere, BlueprintReadWrite, Category = ReverbParameters, meta = (DisplayName = "Decay Time (seconds)", ClampMin = "0.4", ClampMax = "5.0"))
 	float DecayTime;
 
 	/** Room Dampening - 0.0 < 0.85 < 1.0 - Frequency at which the room dampens.*/
@@ -78,6 +78,6 @@ public:
 	UFUNCTION(BlueprintCallable, Category = "Audio|Effects")
 	void SetSettings(const FSubmixEffectFlexiverbSettings& InSettings);
 
-	UPROPERTY(EditAnywhere, BlueprintReadWrite, Category = SubmixEffectPreset)
+	UPROPERTY(EditAnywhere, BlueprintReadWrite, Category = SubmixEffectPreset, meta =(ShowOnlyInnerProperties))
 	FSubmixEffectFlexiverbSettings Settings;
 };
diff --git a/Engine/Plugins/Runtime/Synthesis/Source/Synthesis/Classes/SubmixEffects/SubmixEffectTapDelay.h b/Engine/Plugins/Runtime/Synthesis/Source/Synthesis/Classes/SubmixEffects/SubmixEffectTapDelay.h
--- a/Engine/Plugins/Runtime/Synthesis/Source/Synthesis/Classes/SubmixEffects/SubmixEffectTapDelay.h
+++ b/Engine/Plugins/Runtime/Synthesis/Source/Synthesis/Classes/SubmixEffects/SubmixEffectTapDelay.h
@@ -50,7 +50,7 @@ struct SYNTHESIS_API FTapDelayInfo
 	ETapLineMode TapLineMode;
 
 	// Amount of time before this echo is heard in milliseconds.
-	UPROPERTY(EditAnywhere, BlueprintReadWrite, Category = Realtime, meta = (ClampMin = "0.0", ClampMax = "10000"))
+	UPROPERTY(EditAnywhere, BlueprintReadWrite, Category = Realtime, meta = (DisplayName = "Delay Length (ms)", ClampMin = "0.0", ClampMax = "10000"))
 	float DelayLength;
 
 	// How loud this echo should be, in decibels.
@@ -58,7 +58,7 @@ struct SYNTHESIS_API FTapDelayInfo
 	float Gain;
 
 	// When the Tap Line Mode is set to Send To Channel, designates index of channel from which the echo should play.
-	UPROPERTY(EditAnywhere, BlueprintReadWrite, Category = Realtime, meta = (EditCondition="TapLineMode == ETapLineMode::SendToChannel", ClampMin = "0", ClampMax = "8"))
+	UPROPERTY(EditAnywhere, BlueprintReadWrite, Category = Realtime, meta = (EditCondition = "TapLineMode == ETapLineMode::SendToChannel", ClampMin = "0", ClampMax = "8"))
 	int32 OutputChannel;
 
 	// When the Tap Line Mode is set to Panning, designates the angle at which the echo should be panned.
@@ -86,22 +86,16 @@ struct SYNTHESIS_API FSubmixEffectTapDelaySettings
 	GENERATED_USTRUCT_BODY()
 
 	// Maximum possible length for a delay, in milliseconds. Changing this at runtime will reset the effect.
-	UPROPERTY(EditAnywhere, BlueprintReadWrite, Category = Initialization, meta = (ClampMin = "10.0", UIMin = "10.0", UIMax = "20000.0"))
-	float MaximumDelayLength;
+	UPROPERTY(EditAnywhere, BlueprintReadWrite, Category = Initialization, meta = (DisplayName = "Maximum Delay Length (ms)", ClampMin = "10.0", UIMin = "10.0", UIMax = "20000.0"))
+	float MaximumDelayLength = 10000.0f;
 
 	// Number of milliseconds over which a tap will reach it's set length and gain. Smaller values are more responsive, while larger values will make pitching less dramatic.
-	UPROPERTY(EditAnywhere, BlueprintReadWrite, Category = Realtime, meta = (ClampMin = "0.0", UIMin = "0.0", UIMax = "20000.0"))
-	float InterpolationTime;
+	UPROPERTY(EditAnywhere, BlueprintReadWrite, Category = Realtime, meta = (DisplayName = "Interpolation Time (ms)", ClampMin = "0.0", UIMin = "0.0", UIMax = "20000.0"))
+	float InterpolationTime = 400.0f;
 
 	// Each tap's metadata
 	UPROPERTY(EditAnywhere, BlueprintReadWrite, Category = Realtime)
 	TArray<FTapDelayInfo> Taps;
-
-	FSubmixEffectTapDelaySettings()
-		: MaximumDelayLength(10000.0f)
-		, InterpolationTime(400.0f)
-	{
-	}
 };
 
 class SYNTHESIS_API FSubmixEffectTapDelay : public FSoundEffectSubmix
diff --git a/Engine/Source/Runtime/AudioMixer/Classes/SubmixEffects/AudioMixerSubmixEffectEQ.h b/Engine/Source/Runtime/AudioMixer/Classes/SubmixEffects/AudioMixerSubmixEffectEQ.h
--- a/Engine/Source/Runtime/AudioMixer/Classes/SubmixEffects/AudioMixerSubmixEffectEQ.h
+++ b/Engine/Source/Runtime/AudioMixer/Classes/SubmixEffects/AudioMixerSubmixEffectEQ.h
@@ -18,23 +18,24 @@ struct FSubmixEffectEQBand
 {
 	GENERATED_USTRUCT_BODY()
 
-	UPROPERTY(EditAnywhere, BlueprintReadWrite, Category = "SubmixEffect|Preset", meta = (ClampMin = "20.0", ClampMax = "20000.0", UIMin = "20.0", UIMax = "15000.0"))
-	float Frequency;
+	// Cutoff frequency (in hz) of the EQ band
+	UPROPERTY(EditAnywhere, BlueprintReadWrite, Category = Equalizer, meta = (EditCondition = "bEnabled", DisplayName = "Cutoff Frequency (hz)", ClampMin = "20.0", ClampMax = "20000.0", UIMin = "20.0", UIMax = "15000.0"))
+	float Frequency = 500.0f;
 
-	UPROPERTY(EditAnywhere, BlueprintReadWrite, Category = "SubmixEffect|Preset", meta = (ClampMin = "0.1", ClampMax = "2.0", UIMin = "0.1", UIMax = "2.0"))
-	float Bandwidth;
+	// The bandwidth, in octaves, of the EQ band
+	UPROPERTY(EditAnywhere, BlueprintReadWrite, Category = Equalizer, meta = (EditCondition = "bEnabled", DisplayName = "Bandwidth (octaves)", ClampMin = "0.1", ClampMax = "2.0", UIMin = "0.1", UIMax = "2.0"))
+	float Bandwidth = 2.0f;
 
-	UPROPERTY(EditAnywhere, BlueprintReadWrite, Category = "SubmixEffect|Preset", meta = (DisplayName = "Gain (dB)", ClampMin = "-90.0", ClampMax = "20.0", UIMin = "-90.0", UIMax = "20.0"))
-	float GainDb;
+	// The gain in decibels of the EQ band
+	UPROPERTY(EditAnywhere, BlueprintReadWrite, Category = Equalizer, meta = (EditCondition = "bEnabled", DisplayName = "Gain (dB)", ClampMin = "-90.0", ClampMax = "20.0", UIMin = "-90.0", UIMax = "20.0"))
+	float GainDb = 0.0f;
 
-	UPROPERTY(EditAnywhere, BlueprintReadWrite, Category = "SubmixEffect|Preset")
+	// Whether or not the EQ band is enabled
+	UPROPERTY(EditAnywhere, BlueprintReadWrite, Category = Equalizer)
 	uint32 bEnabled : 1;
 
 	FSubmixEffectEQBand()
-		: Frequency(500.0f)
-		, Bandwidth(2.0f)
-		, GainDb(0.0f)
-		, bEnabled(false)
+		: bEnabled(false)
 	{
 	}
 };
@@ -117,6 +118,6 @@ public:
 	UFUNCTION(BlueprintCallable, Category = "Audio|Effects")
 	AUDIOMIXER_API void SetSettings(const FSubmixEffectSubmixEQSettings& InSettings);
 
-	UPROPERTY(EditAnywhere, BlueprintReadWrite, Category = SubmixEffectPreset)
+	UPROPERTY(EditAnywhere, BlueprintReadWrite, Category = SubmixEffectPreset, meta = (ShowOnlyInnerProperties))
 	FSubmixEffectSubmixEQSettings Settings;
 };
diff --git a/Engine/Source/Runtime/AudioMixer/Classes/SubmixEffects/AudioMixerSubmixEffectReverb.h b/Engine/Source/Runtime/AudioMixer/Classes/SubmixEffects/AudioMixerSubmixEffectReverb.h
--- a/Engine/Source/Runtime/AudioMixer/Classes/SubmixEffects/AudioMixerSubmixEffectReverb.h
+++ b/Engine/Source/Runtime/AudioMixer/Classes/SubmixEffects/AudioMixerSubmixEffectReverb.h
@@ -24,91 +24,69 @@ struct FSubmixEffectReverbSettings
 {
 	GENERATED_USTRUCT_BODY()
 
-
 	/** Bypasses early reflections */
 	UPROPERTY(EditAnywhere, BlueprintReadWrite, Category = EarlyReflections)
-	bool bBypassEarlyReflections;
+	bool bBypassEarlyReflections = false;
 
 	/** Reflections Delay - 0.0 < 0.007 < 0.3 Seconds - the time between the listener receiving the direct path sound and the first reflection */
-	UPROPERTY(EditAnywhere, BlueprintReadWrite, Category = EarlyReflections, meta = (ClampMin = "0.0", ClampMax = "0.3", EditCondition = "!bBypass && !bBypassEarlyReflections"))
-	float ReflectionsDelay;
+	UPROPERTY(EditAnywhere, BlueprintReadWrite, Category = EarlyReflections, meta = (DisplayName = "Reflections Delay (seconds)", ClampMin = "0.0", ClampMax = "0.3", EditCondition = "!bBypass && !bBypassEarlyReflections"))
+	float ReflectionsDelay = 0.007f;
 
 	/** Reverb Gain High Frequency - 0.0 < 0.89 < 1.0 - attenuates the high frequency reflected sound */
-	UPROPERTY(EditAnywhere, BlueprintReadWrite, Category = EarlyReflections, meta = (ClampMin = "0.0", ClampMax = "1.0", EditCondition = "!bBypass && !bBypassEarlyReflections"))
-	float GainHF;
+	UPROPERTY(EditAnywhere, BlueprintReadWrite, Category = EarlyReflections, meta = (DisplayName = "High Frequency Gain", ClampMin = "0.0", ClampMax = "1.0", EditCondition = "!bBypass && !bBypassEarlyReflections"))
+	float GainHF = 0.89f;
 
 	/** Reflections Gain - 0.0 < 0.05 < 3.16 - controls the amount of initial reflections */
 	UPROPERTY(EditAnywhere, BlueprintReadWrite, Category = EarlyReflections, meta = (ClampMin = "0.0", ClampMax = "3.16", EditCondition = "!bBypass && !bBypassEarlyReflections"))
-	float ReflectionsGain;
+	float ReflectionsGain = 0.05f;
 
 	/** Bypasses late reflections. */
 	UPROPERTY(EditAnywhere, BlueprintReadWrite, Category = LateReflections)
-	bool bBypassLateReflections;
+	bool bBypassLateReflections = false;
 
 	/** Late Reverb Delay - 0.0 < 0.011 < 0.1 Seconds - time difference between late reverb and first reflections */
-	UPROPERTY(EditAnywhere, BlueprintReadWrite, Category = LateReflections, meta = (ClampMin = "0.0", ClampMax = "0.1", EditCondition = "!bBypass && !bBypassLateReflections"))
-	float LateDelay;
+	UPROPERTY(EditAnywhere, BlueprintReadWrite, Category = LateReflections, meta = (DisplayName = "Late Delay (seconds)", ClampMin = "0.0", ClampMax = "0.1", EditCondition = "!bBypass && !bBypassLateReflections"))
+	float LateDelay = 0.1f;
 
 	/** Decay Time - 0.1 < 1.49 < 20.0 Seconds - larger is more reverb */
-	UPROPERTY(EditAnywhere, BlueprintReadWrite, Category = LateReflections, meta = (ClampMin = "0.1", ClampMax = "20.0", EditCondition = "!bBypass && !bBypassLateReflections"))
-	float DecayTime;
+	UPROPERTY(EditAnywhere, BlueprintReadWrite, Category = LateReflections, meta = (DisplayName = "Decay Time (seconds)", ClampMin = "0.1", ClampMax = "20.0", EditCondition = "!bBypass && !bBypassLateReflections"))
+	float DecayTime = 1.49f;
 
 	/** Density - 0.0 < 0.85 < 1.0 - Coloration of the late reverb - lower value is more grainy */
 	UPROPERTY(EditAnywhere, BlueprintReadWrite, Category = LateReflections, meta = (ClampMin = "0.0", ClampMax = "1.0", EditCondition = "!bBypass && !bBypassLateReflections"))
-	float Density;
+	float Density = 0.85f;
 
 	/** Diffusion - 0.0 < 0.85 < 1.0 - Echo density in the reverberation decay - lower is more grainy */
 	UPROPERTY(EditAnywhere, BlueprintReadWrite, Category = LateReflections, meta = (ClampMin = "0.0", ClampMax = "1.0", EditCondition = "!bBypass && !bBypassLateReflections"))
-	float Diffusion;
+	float Diffusion = 0.85f;
 
 	/** Air Absorption - 0.0 < 0.994 < 1.0 - lower value means more absorption */
-	UPROPERTY(EditAnywhere, BlueprintReadWrite, Category = LateReflections, meta = (ClampMin = "0.0", ClampMax = "1.0", EditCondition = "!bBypass && !bBypassLateReflections"))
-	float AirAbsorptionGainHF;
+	UPROPERTY(EditAnywhere, BlueprintReadWrite, Category = LateReflections, meta = (DisplayName = "High Frequency Air Absorption Gain", ClampMin = "0.0", ClampMax = "1.0", EditCondition = "!bBypass && !bBypassLateReflections"))
+	float AirAbsorptionGainHF = 0.994f;
 
 	/** Decay High Frequency Ratio - 0.1 < 0.83 < 2.0 - how much quicker or slower the high frequencies decay relative to the lower frequencies. */
-	UPROPERTY(EditAnywhere, BlueprintReadWrite, Category = LateReflections, meta = (ClampMin = "0.1", ClampMax = "2.0", EditCondition = "!bBypass && !bBypassLateReflections"))
-	float DecayHFRatio;
+	UPROPERTY(EditAnywhere, BlueprintReadWrite, Category = LateReflections, meta = (DisplayName = "High Frequency Decay Ratio", ClampMin = "0.1", ClampMax = "2.0", EditCondition = "!bBypass && !bBypassLateReflections"))
+	float DecayHFRatio = 0.83f;
 
 	/** Late Reverb Gain - 0.0 < 1.26 < 10.0 - gain of the late reverb */
-	UPROPERTY(EditAnywhere, BlueprintReadWrite, Category = LateReflections, meta = (ClampMin = "0.0", ClampMax = "10.0", EditCondition = "!bBypass && !bBypassLateReflections"))
-	float LateGain;
+	UPROPERTY(EditAnywhere, BlueprintReadWrite, Category = LateReflections, meta = (DisplayName = "Late Reflections Gain", ClampMin = "0.0", ClampMax = "10.0", EditCondition = "!bBypass && !bBypassLateReflections"))
+	float LateGain = 1.26f;
 
 	/** Reverb Gain - 0.0 < 0.32 < 1.0 - overall reverb gain - master volume control */
 	UPROPERTY(EditAnywhere, BlueprintReadWrite, Category = LateReflections, meta = (ClampMin = "0.0", ClampMax = "1.0", EditCondition = "!bBypass && !bBypassLateReflections"))
-	float Gain;
+	float Gain = 0.0f;
 
 	// Overall wet level of the reverb effect
 	UPROPERTY(EditAnywhere, BlueprintReadWrite, Category = Routing, meta = (EditCondition = "!bBypass", UIMin = "0.0", UIMax = "1.0", ClampMin = "0.0", ClampMax = "10.0"))
-	float WetLevel;
+	float WetLevel = 0.3f;
 
 	// Overall dry level of the reverb effect
 	UPROPERTY(EditAnywhere, BlueprintReadWrite, Category = Routing, meta = (EditCondition = "!bBypass", ClampMin = "0.0", ClampMax = "1.0"))
-	float DryLevel;
+	float DryLevel = 0.0f;
 
 	/** Bypasses reverb */
 	UPROPERTY(EditAnywhere, BlueprintReadWrite, Category = General)
-	bool bBypass;
-
-
-	FSubmixEffectReverbSettings()
-		: bBypassEarlyReflections(false)
-		, ReflectionsDelay(0.007f)
-		, GainHF(0.89f)
-		, ReflectionsGain(0.05f)
-		, bBypassLateReflections(false)
-		, LateDelay(0.1f)
-		, DecayTime(1.49f)
-		, Density(0.85f)
-		, Diffusion(0.85f)
-		, AirAbsorptionGainHF(0.994f)
-		, DecayHFRatio(0.83f)
-		, LateGain(1.26f)
-		, Gain(0.0f)
-		, WetLevel(0.3f)
-		, DryLevel(0.0f)
-		, bBypass(false)
-	{
-	}
+	bool bBypass = false;
 };
 
 class FSubmixEffectReverb : public FSoundEffectSubmix
@@ -181,6 +159,6 @@ public:
 	UFUNCTION(BlueprintCallable, Category = "Audio|Effects")
 	AUDIOMIXER_API void SetSettingsWithReverbEffect(const UReverbEffect* InReverbEffect, const float WetLevel, const float DryLevel = 0.0f);
 
-	UPROPERTY(EditAnywhere, BlueprintReadWrite, Category = SubmixEffectPreset)
+	UPROPERTY(EditAnywhere, BlueprintReadWrite, Category = SubmixEffectPreset, meta = (ShowOnlyInnerProperties))
 	FSubmixEffectReverbSettings Settings;
 };
