Adding command line and cvar options to disable hardware convolutioin.

[RN] Added ability to disable hardware accelerated convolution by commandline or cvar

diff --git a/Engine/Plugins/Runtime/Synthesis/Source/Synthesis/Private/EffectConvolutionReverb.cpp b/Engine/Plugins/Runtime/Synthesis/Source/Synthesis/Private/EffectConvolutionReverb.cpp
--- a/Engine/Plugins/Runtime/Synthesis/Source/Synthesis/Private/EffectConvolutionReverb.cpp
+++ b/Engine/Plugins/Runtime/Synthesis/Source/Synthesis/Private/EffectConvolutionReverb.cpp
@@ -4,9 +4,31 @@
 #include "Async/AsyncWork.h"
 #include "DSP/AlignedBlockBuffer.h" // IWYU pragma: keep
 #include "SynthesisModule.h"
+#include "Misc/Parse.h"
+#include "Misc/CommandLine.h"
+#include "HAL/IConsoleManager.h"
 
 #include UE_INLINE_GENERATED_CPP_BY_NAME(EffectConvolutionReverb)
 
+
+static int32 bDisabledHardwareConvolution = 0;
+FAutoConsoleVariableRef CVarDisableHardwareConvolution(
+	TEXT("au.DisableHardwareConvolution"),
+	bDisabledHardwareConvolution,
+	TEXT("Disables hardware convolution. Will not change already instantiated hardware accelerated convolution.\n")
+	TEXT("1: New instances of hardware-acclerated convolution is disabled. 0: Hardware convolution is not disabled."),
+	ECVF_Default);
+
+
+// Command line arg to disable hardware convolution
+static bool bDisableConvolutionHardareAcceleration = FParse::Param(FCommandLine::Get(), TEXT("DisableHardwareConvolution"));
+
+static bool IsHardwareConvolutionDisabled()
+{
+	return bDisableConvolutionHardareAcceleration || bDisabledHardwareConvolution;
+}
+
+
 UAudioImpulseResponse::UAudioImpulseResponse()
 	: NumChannels(0)
 	, SampleRate(0)
@@ -191,6 +213,15 @@ namespace Audio
 			}
 		}
 
+		if (IsHardwareConvolutionDisabled())
+		{
+			if (ConvReverbInitData.AlgorithmSettings.bEnableHardwareAcceleration)
+			{
+				UE_LOG(LogSynthesis, Display, TEXT("Disabling hardware acceleration due to cvar or command line argument."));
+				ConvReverbInitData.AlgorithmSettings.bEnableHardwareAcceleration = false;
+			}
+		}
+
 		return ConvReverbInitData;
 	}
 
