Fixing command line arg disabling convolution reverb by moving the static int to function local scope

#rnx

diff --git a/Engine/Plugins/Runtime/Synthesis/Source/Synthesis/Private/EffectConvolutionReverb.cpp b/Engine/Plugins/Runtime/Synthesis/Source/Synthesis/Private/EffectConvolutionReverb.cpp
--- a/Engine/Plugins/Runtime/Synthesis/Source/Synthesis/Private/EffectConvolutionReverb.cpp
+++ b/Engine/Plugins/Runtime/Synthesis/Source/Synthesis/Private/EffectConvolutionReverb.cpp
@@ -21,10 +21,10 @@ FAutoConsoleVariableRef CVarDisableHardwareConvolution(
 
 
 // Command line arg to disable hardware convolution
-static bool bDisableConvolutionHardareAcceleration = FParse::Param(FCommandLine::Get(), TEXT("DisableHardwareConvolution"));
 
 static bool IsHardwareConvolutionDisabled()
-{
+{	
+	static bool bDisableConvolutionHardareAcceleration = FParse::Param(FCommandLine::Get(), TEXT("DisableHardwareConvolution"));
 	return bDisableConvolutionHardareAcceleration || bDisabledHardwareConvolution;
 }
 
