Renaming EQ source effect Enable property to reflect its current behavior. Enabling Dynamic source effect Bypass property to function properly.

diff --git a/Engine/Plugins/Runtime/Synthesis/Source/Synthesis/Classes/SourceEffects/SourceEffectEQ.h b/Engine/Plugins/Runtime/Synthesis/Source/Synthesis/Classes/SourceEffects/SourceEffectEQ.h
--- a/Engine/Plugins/Runtime/Synthesis/Source/Synthesis/Classes/SourceEffects/SourceEffectEQ.h
+++ b/Engine/Plugins/Runtime/Synthesis/Source/Synthesis/Classes/SourceEffects/SourceEffectEQ.h
@@ -16,19 +16,19 @@ struct FSourceEffectEQBand
 	GENERATED_USTRUCT_BODY()
 
 	// The cutoff frequency of the band
-	UPROPERTY(EditAnywhere, BlueprintReadWrite, Category = Equalizer, meta = (DisplayName = "Frequency (hz)", ClampMin = "20.0", ClampMax = "20000.0", UIMin = "0.0", UIMax = "15000.0", EditCondition = "!bEnabled"))
+	UPROPERTY(EditAnywhere, BlueprintReadWrite, Category = Equalizer, meta = (DisplayName = "Frequency (hz)", ClampMin = "20.0", ClampMax = "20000.0", UIMin = "0.0", UIMax = "15000.0", EditCondition = "bEnabled"))
 	float Frequency;
 
 	// The bandwidth (in octaves) of the band
-	UPROPERTY(EditAnywhere, BlueprintReadWrite, Category = Equalizer, meta = (DisplayName = "Bandwidth (octaves)", ClampMin = "0.1", ClampMax = "2.0", UIMin = "0.1", UIMax = "2.0", EditCondition = "!bEnabled"))
+	UPROPERTY(EditAnywhere, BlueprintReadWrite, Category = Equalizer, meta = (DisplayName = "Bandwidth (octaves)", ClampMin = "0.1", ClampMax = "2.0", UIMin = "0.1", UIMax = "2.0", EditCondition = "bEnabled"))
 	float Bandwidth;
 
 	// The gain in decibels to apply to the eq band
-	UPROPERTY(EditAnywhere, BlueprintReadWrite, Category = Equalizer, meta = (DisplayName = "Gain (dB)", ClampMin = "-90.0", ClampMax = "20.0", UIMin = "-90.0", UIMax = "20.0", EditCondition = "!bEnabled"))
+	UPROPERTY(EditAnywhere, BlueprintReadWrite, Category = Equalizer, meta = (DisplayName = "Gain (dB)", ClampMin = "-90.0", ClampMax = "20.0", UIMin = "-90.0", UIMax = "20.0", EditCondition = "bEnabled"))
 	float GainDb;
 
 	// Whether or not the band is enabled. Allows changing bands on the fly.
-	UPROPERTY(EditAnywhere, BlueprintReadWrite, Category = Equalizer, meta = (DisplayName = "Bypass"))
+	UPROPERTY(EditAnywhere, BlueprintReadWrite, Category = Equalizer, meta = (DisplayName = "Enabled"))
 	uint32 bEnabled : 1;
 
 	FSourceEffectEQBand()
diff --git a/Engine/Plugins/Runtime/Synthesis/Source/Synthesis/Private/SourceEffectDynamicsProcessor.cpp b/Engine/Plugins/Runtime/Synthesis/Source/Synthesis/Private/SourceEffectDynamicsProcessor.cpp
--- a/Engine/Plugins/Runtime/Synthesis/Source/Synthesis/Private/SourceEffectDynamicsProcessor.cpp
+++ b/Engine/Plugins/Runtime/Synthesis/Source/Synthesis/Private/SourceEffectDynamicsProcessor.cpp
@@ -68,6 +68,7 @@ void FSourceEffectDynamicsProcessor::OnPresetChanged()
 	DynamicsProcessor.SetInputGain(Settings.InputGainDb);
 	DynamicsProcessor.SetOutputGain(Settings.OutputGainDb);
 	DynamicsProcessor.SetAnalogMode(Settings.bAnalogMode);
+	DynamicsProcessor.SetKeyAudition(Settings.bBypass);
 
 	DynamicsProcessor.SetChannelLinkMode(Settings.bStereoLinked ? Audio::EDynamicsProcessorChannelLinkMode::Average : Audio::EDynamicsProcessorChannelLinkMode::Disabled);
 }
