Remove source effect compressor bypass change.

diff --git a/Engine/Plugins/Runtime/Synthesis/Source/Synthesis/Private/SourceEffectDynamicsProcessor.cpp b/Engine/Plugins/Runtime/Synthesis/Source/Synthesis/Private/SourceEffectDynamicsProcessor.cpp
--- a/Engine/Plugins/Runtime/Synthesis/Source/Synthesis/Private/SourceEffectDynamicsProcessor.cpp
+++ b/Engine/Plugins/Runtime/Synthesis/Source/Synthesis/Private/SourceEffectDynamicsProcessor.cpp
@@ -68,7 +68,6 @@ void FSourceEffectDynamicsProcessor::OnPresetChanged()
 	DynamicsProcessor.SetInputGain(Settings.InputGainDb);
 	DynamicsProcessor.SetOutputGain(Settings.OutputGainDb);
 	DynamicsProcessor.SetAnalogMode(Settings.bAnalogMode);
-	DynamicsProcessor.SetKeyAudition(Settings.bBypass);
 
 	DynamicsProcessor.SetChannelLinkMode(Settings.bStereoLinked ? Audio::EDynamicsProcessorChannelLinkMode::Average : Audio::EDynamicsProcessorChannelLinkMode::Disabled);
 }
