Fix Chorus Modulation

diff --git a/Engine/Plugins/Runtime/Synthesis/Source/Synthesis/Private/SourceEffectChorus.cpp b/Engine/Plugins/Runtime/Synthesis/Source/Synthesis/Private/SourceEffectChorus.cpp
--- a/Engine/Plugins/Runtime/Synthesis/Source/Synthesis/Private/SourceEffectChorus.cpp
+++ b/Engine/Plugins/Runtime/Synthesis/Source/Synthesis/Private/SourceEffectChorus.cpp
@@ -19,6 +19,20 @@ void FSourceEffectChorus::Init(const FSoundEffectSourceInitData& InitData)
 		WetMod.Init(DeviceId, bIsBuffered);
 		DryMod.Init(DeviceId, bIsBuffered);
 		SpreadMod.Init(DeviceId, bIsBuffered);
+
+		auto TransformObjectPtrSet = [](const TSet<TObjectPtr<USoundModulatorBase>>& InPtrs)
+		{
+			TSet<USoundModulatorBase*> Objects;
+			Algo::Transform(InPtrs, Objects, [](const TObjectPtr<USoundModulatorBase>& Obj) { return Obj; });
+			return Objects;
+		};
+
+		SetDepthModulators(TransformObjectPtrSet(ProcPreset->Settings.DepthModulation.Modulators));
+		SetFeedbackModulators(TransformObjectPtrSet(ProcPreset->Settings.FeedbackModulation.Modulators));
+		SetFrequencyModulators(TransformObjectPtrSet(ProcPreset->Settings.FrequencyModulation.Modulators));
+		SetWetModulators(TransformObjectPtrSet(ProcPreset->Settings.WetModulation.Modulators));
+		SetDryModulators(TransformObjectPtrSet(ProcPreset->Settings.DryModulation.Modulators));
+		SetSpreadModulators(TransformObjectPtrSet(ProcPreset->Settings.SpreadModulation.Modulators));
 	}
 
 	Chorus.Init(InitData.SampleRate, InitData.NumSourceChannels, 2.0f, 64);
