Rename FModulationDestination::UpdateModulators to SetModulators
#Synced-CL 51549379

diff --git a/Engine/Plugins/Runtime/AudioModulation/Source/AudioModulation/Private/AudioModulationDestination.cpp b/Engine/Plugins/Runtime/AudioModulation/Source/AudioModulation/Private/AudioModulationDestination.cpp
--- a/Engine/Plugins/Runtime/AudioModulation/Source/AudioModulation/Private/AudioModulationDestination.cpp
+++ b/Engine/Plugins/Runtime/AudioModulation/Source/AudioModulation/Private/AudioModulationDestination.cpp
@@ -53,7 +53,7 @@ bool UAudioModulationDestination::ClearModulator()
 	if (Modulator)
 	{
 		Modulator = nullptr;
-		Destination.UpdateModulators(TSet<const USoundModulatorBase*>{ });
+		Destination.SetModulators(TSet<const USoundModulatorBase*>{ });
 		return true;
 	}
 
@@ -109,7 +109,7 @@ bool UAudioModulationDestination::SetModulator(const USoundModulatorBase* InModu
 	if (InModulator)
 	{
 		Modulator = InModulator;
-		Destination.UpdateModulators(TSet<const USoundModulatorBase*>{ InModulator });
+		Destination.SetModulators(TSet<const USoundModulatorBase*>{ InModulator });
 		return true;
 	}
 
diff --git a/Engine/Plugins/Runtime/Synthesis/Source/Synthesis/Private/SourceEffectBitCrusher.cpp b/Engine/Plugins/Runtime/Synthesis/Source/Synthesis/Private/SourceEffectBitCrusher.cpp
--- a/Engine/Plugins/Runtime/Synthesis/Source/Synthesis/Private/SourceEffectBitCrusher.cpp
+++ b/Engine/Plugins/Runtime/Synthesis/Source/Synthesis/Private/SourceEffectBitCrusher.cpp
@@ -52,22 +52,22 @@ void FSourceEffectBitCrusher::ProcessAudio(const FSoundEffectSourceInputData& In
 
 void FSourceEffectBitCrusher::SetBitModulator(const USoundModulatorBase* InModulator)
 {
-	BitMod.UpdateModulators({ InModulator });
+	BitMod.SetModulators({ InModulator });
 }
 
 void FSourceEffectBitCrusher::SetBitModulators(const TSet<USoundModulatorBase*>& InModulators)
 {
-	BitMod.UpdateModulators(InModulators);
+	BitMod.SetModulators(InModulators);
 }
 
 void FSourceEffectBitCrusher::SetSampleRateModulator(const USoundModulatorBase* InModulator)
 {
-	SampleRateMod.UpdateModulators({ InModulator });
+	SampleRateMod.SetModulators({ InModulator });
 }
 
 void FSourceEffectBitCrusher::SetSampleRateModulators(const TSet<USoundModulatorBase*>& InModulators)
 {
-	SampleRateMod.UpdateModulators(InModulators);
+	SampleRateMod.SetModulators(InModulators);
 }
 
 void USourceEffectBitCrusherPreset::OnInit()
diff --git a/Engine/Plugins/Runtime/Synthesis/Source/Synthesis/Private/SourceEffectChorus.cpp b/Engine/Plugins/Runtime/Synthesis/Source/Synthesis/Private/SourceEffectChorus.cpp
--- a/Engine/Plugins/Runtime/Synthesis/Source/Synthesis/Private/SourceEffectChorus.cpp
+++ b/Engine/Plugins/Runtime/Synthesis/Source/Synthesis/Private/SourceEffectChorus.cpp
@@ -91,62 +91,62 @@ void FSourceEffectChorus::ProcessAudio(const FSoundEffectSourceInputData& InData
 
 void FSourceEffectChorus::SetDepthModulator(const USoundModulatorBase* InModulator)
 {
-	DepthMod.UpdateModulators({ InModulator });
+	DepthMod.SetModulators({ InModulator });
 }
 
 void FSourceEffectChorus::SetFeedbackModulator(const USoundModulatorBase* InModulator)
 {
-	FeedbackMod.UpdateModulators({ InModulator });
+	FeedbackMod.SetModulators({ InModulator });
 }
 
 void FSourceEffectChorus::SetSpreadModulator(const USoundModulatorBase* InModulator)
 {
-	SpreadMod.UpdateModulators({ InModulator });
+	SpreadMod.SetModulators({ InModulator });
 }
 
 void FSourceEffectChorus::SetDryModulator(const USoundModulatorBase* InModulator)
 {
-	DryMod.UpdateModulators({ InModulator });
+	DryMod.SetModulators({ InModulator });
 }
 
 void FSourceEffectChorus::SetWetModulator(const USoundModulatorBase* InModulator)
 {
-	WetMod.UpdateModulators({ InModulator });
+	WetMod.SetModulators({ InModulator });
 }
 
 void FSourceEffectChorus::SetFrequencyModulator(const USoundModulatorBase* InModulator)
 {
-	FrequencyMod.UpdateModulators({ InModulator });
+	FrequencyMod.SetModulators({ InModulator });
 }
 
 void FSourceEffectChorus::SetDepthModulators(const TSet<USoundModulatorBase*>& InModulators)
 {
-	DepthMod.UpdateModulators(InModulators);
+	DepthMod.SetModulators(InModulators);
 }
 
 void FSourceEffectChorus::SetFeedbackModulators(const TSet<USoundModulatorBase*>& InModulators)
 {
-	FeedbackMod.UpdateModulators(InModulators);
+	FeedbackMod.SetModulators(InModulators);
 }
 
 void FSourceEffectChorus::SetSpreadModulators(const TSet<USoundModulatorBase*>& InModulators)
 {
-	SpreadMod.UpdateModulators(InModulators);
+	SpreadMod.SetModulators(InModulators);
 }
 
 void FSourceEffectChorus::SetDryModulators(const TSet<USoundModulatorBase*>& InModulators)
 {
-	DryMod.UpdateModulators(InModulators);
+	DryMod.SetModulators(InModulators);
 }
 
 void FSourceEffectChorus::SetWetModulators(const TSet<USoundModulatorBase*>& InModulators)
 {
-	WetMod.UpdateModulators(InModulators);
+	WetMod.SetModulators(InModulators);
 }
 
 void FSourceEffectChorus::SetFrequencyModulators(const TSet<USoundModulatorBase*>& InModulators)
 {
-	FrequencyMod.UpdateModulators(InModulators);
+	FrequencyMod.SetModulators(InModulators);
 }
 
 void USourceEffectChorusPreset::OnInit()
@@ -158,7 +158,7 @@ void USourceEffectChorusPreset::OnInit()
 		return Objects;
 	};
 
-	SetDepthModulators(TransformObjectPtrSet(Settings.SpreadModulation.Modulators));
+	SetDepthModulators(TransformObjectPtrSet(Settings.DepthModulation.Modulators));
 	SetDryModulators(TransformObjectPtrSet(Settings.DryModulation.Modulators));
 	SetWetModulators(TransformObjectPtrSet(Settings.WetModulation.Modulators));
 	SetFeedbackModulators(TransformObjectPtrSet(Settings.FeedbackModulation.Modulators));
diff --git a/Engine/Source/Runtime/AudioMixer/Private/AudioMixerSourceManager.cpp b/Engine/Source/Runtime/AudioMixer/Private/AudioMixerSourceManager.cpp
--- a/Engine/Source/Runtime/AudioMixer/Private/AudioMixerSourceManager.cpp
+++ b/Engine/Source/Runtime/AudioMixer/Private/AudioMixerSourceManager.cpp
@@ -1248,19 +1248,19 @@ namespace Audio
 
 		FModulationDestination VolumeMod;
 		VolumeMod.Init(MixerDevice->DeviceID, FName("Volume"), false /* bInIsBuffered */, true /* bInValueLinear */);
-		VolumeMod.UpdateModulators(InitParams.ModulationSettings.VolumeModulationDestination.Modulators);
+		VolumeMod.SetModulators(InitParams.ModulationSettings.VolumeModulationDestination.Modulators);
 
 		FModulationDestination PitchMod;
 		PitchMod.Init(MixerDevice->DeviceID, FName("Pitch"), false /* bInIsBuffered */);
-		PitchMod.UpdateModulators(InitParams.ModulationSettings.PitchModulationDestination.Modulators);
+		PitchMod.SetModulators(InitParams.ModulationSettings.PitchModulationDestination.Modulators);
 
 		FModulationDestination HighpassMod;
 		HighpassMod.Init(MixerDevice->DeviceID, FName("HPFCutoffFrequency"), false /* bInIsBuffered */);
-		HighpassMod.UpdateModulators(InitParams.ModulationSettings.HighpassModulationDestination.Modulators);
+		HighpassMod.SetModulators(InitParams.ModulationSettings.HighpassModulationDestination.Modulators);
 
 		FModulationDestination LowpassMod;
 		LowpassMod.Init(MixerDevice->DeviceID, FName("LPFCutoffFrequency"), false /* bInIsBuffered */);
-		LowpassMod.UpdateModulators(InitParams.ModulationSettings.LowpassModulationDestination.Modulators);
+		LowpassMod.SetModulators(InitParams.ModulationSettings.LowpassModulationDestination.Modulators);
 
 		AudioMixerThreadCommand([
 			this,
@@ -2075,19 +2075,19 @@ namespace Audio
 
 		FModulationDestination VolumeMod;
 		VolumeMod.Init(MixerDevice->DeviceID, FName("Volume"), false /* bInIsBuffered */, true /* bInValueLinear */);
-		VolumeMod.UpdateModulators(ModulationSettings.VolumeModulationDestination.Modulators);
+		VolumeMod.SetModulators(ModulationSettings.VolumeModulationDestination.Modulators);
 
 		FModulationDestination PitchMod;
 		PitchMod.Init(MixerDevice->DeviceID, FName("Pitch"), false /* bInIsBuffered */);
-		PitchMod.UpdateModulators(ModulationSettings.PitchModulationDestination.Modulators);
+		PitchMod.SetModulators(ModulationSettings.PitchModulationDestination.Modulators);
 
 		FModulationDestination HighpassMod;
 		HighpassMod.Init(MixerDevice->DeviceID, FName("HPFCutoffFrequency"), false /* bInIsBuffered */);
-		HighpassMod.UpdateModulators(ModulationSettings.HighpassModulationDestination.Modulators);
+		HighpassMod.SetModulators(ModulationSettings.HighpassModulationDestination.Modulators);
 
 		FModulationDestination LowpassMod;
 		LowpassMod.Init(MixerDevice->DeviceID, FName("LPFCutoffFrequency"), false /* bInIsBuffered */);
-		LowpassMod.UpdateModulators(ModulationSettings.LowpassModulationDestination.Modulators);
+		LowpassMod.SetModulators(ModulationSettings.LowpassModulationDestination.Modulators);
 
 
 		AudioMixerThreadCommand([
diff --git a/Engine/Source/Runtime/AudioMixer/Private/AudioMixerSubmix.cpp b/Engine/Source/Runtime/AudioMixer/Private/AudioMixerSubmix.cpp
--- a/Engine/Source/Runtime/AudioMixer/Private/AudioMixerSubmix.cpp
+++ b/Engine/Source/Runtime/AudioMixer/Private/AudioMixerSubmix.cpp
@@ -2631,9 +2631,9 @@ namespace Audio
 
 	void FMixerSubmix::UpdateModulationSettings(const TSet<TObjectPtr<USoundModulatorBase>>& InOutputModulators, const TSet<TObjectPtr<USoundModulatorBase>>& InWetLevelModulators, const TSet<TObjectPtr<USoundModulatorBase>>& InDryLevelModulators)
 	{
-		VolumeMod.UpdateModulators(InOutputModulators);
-		WetLevelMod.UpdateModulators(InWetLevelModulators);
-		DryLevelMod.UpdateModulators(InDryLevelModulators);
+		VolumeMod.SetModulators(InOutputModulators);
+		WetLevelMod.SetModulators(InWetLevelModulators);
+		DryLevelMod.SetModulators(InDryLevelModulators);
 	}
 
 	void FMixerSubmix::SetModulationBaseLevels(float InVolumeModBaseDb, float InWetModBaseDb, float InDryModBaseDb)
diff --git a/Engine/Source/Runtime/Engine/Classes/Sound/SoundModulationDestination.h b/Engine/Source/Runtime/Engine/Classes/Sound/SoundModulationDestination.h
--- a/Engine/Source/Runtime/Engine/Classes/Sound/SoundModulationDestination.h
+++ b/Engine/Source/Runtime/Engine/Classes/Sound/SoundModulationDestination.h
@@ -192,12 +192,22 @@ namespace Audio
 			ENGINE_API bool ProcessControl(float InValueUnitBase, int32 InNumSamples = 0);
 
 
-			UE_DEPRECATED(5.1, "Deprecated in favor of supporting multiple modulators per destination. Use 'UpdateModulators' instead.")
+			UE_DEPRECATED(5.1, "Deprecated in favor of supporting multiple modulators per destination. Use 'SetModulators' instead.")
 			ENGINE_API void UpdateModulator(const USoundModulatorBase* InModulator);
 
+			UE_DEPRECATED(5.8, "This function has been renamed SetModulators.")
 			ENGINE_API void UpdateModulators(const TSet<TObjectPtr<USoundModulatorBase>>& InModulators);
+
+			UE_DEPRECATED(5.8, "This function has been renamed SetModulators.")
 			ENGINE_API void UpdateModulators(const TSet<USoundModulatorBase*>& InModulators);
+
+			UE_DEPRECATED(5.8, "This function has been renamed SetModulators.")
 			ENGINE_API void UpdateModulators(const TSet<const USoundModulatorBase*>& InModulators);
+
+			/**	Updates the set of Modulators on this Destination. */
+			ENGINE_API void SetModulators(const TSet<TObjectPtr<USoundModulatorBase>>& InModulators);
+			ENGINE_API void SetModulators(const TSet<USoundModulatorBase*>& InModulators);
+			ENGINE_API void SetModulators(const TSet<const USoundModulatorBase*>& InModulators);
 		
 			ENGINE_API void ClearModulators();
 
diff --git a/Engine/Source/Runtime/Engine/Private/AudioVirtualLoop.cpp b/Engine/Source/Runtime/Engine/Private/AudioVirtualLoop.cpp
--- a/Engine/Source/Runtime/Engine/Private/AudioVirtualLoop.cpp
+++ b/Engine/Source/Runtime/Engine/Private/AudioVirtualLoop.cpp
@@ -153,7 +153,7 @@ bool FAudioVirtualLoop::Virtualize(const FActiveSound& InActiveSound, FAudioDevi
 				if (!PitchModulationSettings.Modulators.IsEmpty())
 				{
 					NewWaveInstanceData.PitchModulationDestination.Init(InAudioDevice.DeviceID, AudioVirtualLoopPrivate::PitchModulationParameterName);
-					NewWaveInstanceData.PitchModulationDestination.UpdateModulators(PitchModulationSettings.Modulators);
+					NewWaveInstanceData.PitchModulationDestination.SetModulators(PitchModulationSettings.Modulators);
 				}
 				
 				OutVirtualLoop.WaveInstanceData.Add(Pair.Key, NewWaveInstanceData);
@@ -171,7 +171,7 @@ bool FAudioVirtualLoop::Virtualize(const FActiveSound& InActiveSound, FAudioDevi
 	{
 		OutVirtualLoop.VolumeConcurrencyDestinations.Add(Audio::FModulationDestination());
 		OutVirtualLoop.VolumeConcurrencyDestinations[DestinationIndex].Init(InAudioDevice.DeviceID, AudioVirtualLoopPrivate::VolumeModulationParameterName, false /* bInIsBuffered */, true /* bInValueNormalized */);
-		OutVirtualLoop.VolumeConcurrencyDestinations[DestinationIndex].UpdateModulators(ModulatorSets[DestinationIndex]);
+		OutVirtualLoop.VolumeConcurrencyDestinations[DestinationIndex].SetModulators(ModulatorSets[DestinationIndex]);
 	}
 
 #if UE_AUDIO_PROFILERTRACE_ENABLED
diff --git a/Engine/Source/Runtime/Engine/Private/SoundModulationDestination.cpp b/Engine/Source/Runtime/Engine/Private/SoundModulationDestination.cpp
--- a/Engine/Source/Runtime/Engine/Private/SoundModulationDestination.cpp
+++ b/Engine/Source/Runtime/Engine/Private/SoundModulationDestination.cpp
@@ -254,10 +254,25 @@ namespace Audio
 
 	void FModulationDestination::UpdateModulator(const USoundModulatorBase* InModulator)
 	{
-		UpdateModulators({ InModulator });
+		SetModulators({ InModulator });
 	}
 
 	void FModulationDestination::UpdateModulators(const TSet<TObjectPtr<USoundModulatorBase>>& InModulators)
+	{
+		SetModulators(InModulators);
+	}
+
+	void FModulationDestination::UpdateModulators(const TSet<USoundModulatorBase*>& InModulators)
+	{
+		SetModulators(InModulators);
+	}
+
+	void FModulationDestination::UpdateModulators(const TSet<const USoundModulatorBase*>& InModulators)
+	{
+		SetModulators(InModulators);
+	}
+
+	void FModulationDestination::SetModulators(const TSet<TObjectPtr<USoundModulatorBase>>& InModulators)
 	{
 		TArray<TUniquePtr<Audio::IModulatorSettings>> ProxySettings;
 		Algo::TransformIf(
@@ -270,7 +285,7 @@ namespace Audio
 		UpdateModulatorsInternal(MoveTemp(ProxySettings));
 	}
 
-	void FModulationDestination::UpdateModulators(const TSet<USoundModulatorBase*>& InModulators)
+	void FModulationDestination::SetModulators(const TSet<USoundModulatorBase*>& InModulators)
 	{
 		TArray<TUniquePtr<Audio::IModulatorSettings>> ProxySettings;
 		Algo::TransformIf(
@@ -283,7 +298,7 @@ namespace Audio
 		UpdateModulatorsInternal(MoveTemp(ProxySettings));
 	}
 
-	void FModulationDestination::UpdateModulators(const TSet<const USoundModulatorBase*>& InModulators)
+	void FModulationDestination::SetModulators(const TSet<const USoundModulatorBase*>& InModulators)
 	{
 		TArray<TUniquePtr<Audio::IModulatorSettings>> ProxySettings;
 		Algo::TransformIf(
